[SlateInspector] Move `SlateInspector` plugin to `Engine/Plugins/Experimental`.

#Synced-CL 51849904

diff --git a/Engine/Plugins/Experimental/SlateInspector/SlateInspector.uplugin b/Engine/Plugins/Experimental/SlateInspector/SlateInspector.uplugin
new file mode 100644
--- /dev/null
+++ b/Engine/Plugins/Experimental/SlateInspector/SlateInspector.uplugin
@@ -0,0 +1,31 @@
+{
+	"FileVersion": 3,
+	"Version": 1,
+	"VersionName": "1.0",
+	"FriendlyName": "SlateInspector",
+	"Description": "Slate UI automation and inspection tools.",
+	"NoRedist": true,
+	"EditorOnly": true,
+	"CreatedBy": "",
+	"CreatedByURL": "",
+	"DocsURL": "",
+	"MarketplaceURL": "",
+	"CanContainContent": false,
+	"IsBetaVersion": false,
+	"IsExperimentalVersion": true,
+	"Installed": false,
+	"EnabledByDefault": false,
+	"Modules": [
+		{
+			"Name": "SlateInspector",
+			"Type": "Editor",
+			"LoadingPhase": "Default"
+		}
+	],
+	"Plugins": [
+		{
+			"Name": "ToolsetRegistry",
+			"Enabled": true
+		}
+	]
+}
diff --git a/Engine/Plugins/Experimental/SlateInspector/Source/SlateInspector/Private/SlateInspectorModule.cpp b/Engine/Plugins/Experimental/SlateInspector/Source/SlateInspector/Private/SlateInspectorModule.cpp
new file mode 100644
--- /dev/null
+++ b/Engine/Plugins/Experimental/SlateInspector/Source/SlateInspector/Private/SlateInspectorModule.cpp
@@ -0,0 +1,15 @@
+// Copyright Epic Games, Inc. All Rights Reserved.
+
+#include "SlateInspectorModule.h"
+
+#include "Modules/ModuleManager.h"
+
+IMPLEMENT_MODULE(FSlateInspectorModule, SlateInspector)
+
+void FSlateInspectorModule::StartupModule()
+{
+}
+
+void FSlateInspectorModule::ShutdownModule()
+{
+}
diff --git a/Engine/Plugins/Experimental/SlateInspector/Source/SlateInspector/Private/SlateInspectorRefCache.cpp b/Engine/Plugins/Experimental/SlateInspector/Source/SlateInspector/Private/SlateInspectorRefCache.cpp
new file mode 100644
--- /dev/null
+++ b/Engine/Plugins/Experimental/SlateInspector/Source/SlateInspector/Private/SlateInspectorRefCache.cpp
@@ -0,0 +1,45 @@
+// Copyright Epic Games, Inc. All Rights Reserved.
+
+#include "SlateInspectorRefCache.h"
+
+FSlateInspectorRefCache& FSlateInspectorRefCache::Get()
+{
+	static FSlateInspectorRefCache Instance;
+	return Instance;
+}
+
+FString FSlateInspectorRefCache::GetOrAssignRef(TSharedRef<SWidget> Widget, const FString& RolePrefix)
+{
+	TWeakPtr<SWidget> WeakWidget = Widget;
+
+	if (const FString* ExistingRef = WidgetToRef.Find(WeakWidget))
+	{
+		return *ExistingRef;
+	}
+
+	int32& Counter = RolePrefixCounters.FindOrAdd(RolePrefix, 0);
+	++Counter;
+
+	FString Ref = FString::Printf(TEXT("%s%d"), *RolePrefix, Counter);
+
+	WidgetToRef.Add(WeakWidget, Ref);
+	RefToWidget.Add(Ref, WeakWidget);
+
+	return Ref;
+}
+
+TSharedPtr<SWidget> FSlateInspectorRefCache::ResolveRef(const FString& Ref) const
+{
+	if (const TWeakPtr<SWidget>* WeakWidget = RefToWidget.Find(Ref))
+	{
+		return WeakWidget->Pin();
+	}
+	return nullptr;
+}
+
+void FSlateInspectorRefCache::Reset()
+{
+	WidgetToRef.Empty();
+	RefToWidget.Empty();
+	RolePrefixCounters.Empty();
+}
diff --git a/Engine/Plugins/Experimental/SlateInspector/Source/SlateInspector/Private/SlateInspectorSnapshotRenderer.cpp b/Engine/Plugins/Experimental/SlateInspector/Source/SlateInspector/Private/SlateInspectorSnapshotRenderer.cpp
new file mode 100644
--- /dev/null
+++ b/Engine/Plugins/Experimental/SlateInspector/Source/SlateInspector/Private/SlateInspectorSnapshotRenderer.cpp
@@ -0,0 +1,406 @@
+// Copyright Epic Games, Inc. All Rights Reserved.
+
+#include "SlateInspectorSnapshotRenderer.h"
+
+#include "SlateInspectorRefCache.h"
+
+#include "Framework/Application/SlateApplication.h"
+#include "Misc/Paths.h"
+#include "Widgets/Input/SCheckBox.h"
+#include "Widgets/Input/SEditableTextBox.h"
+#include "Widgets/Docking/SDockTab.h"
+#include "Widgets/SWindow.h"
+#include "Widgets/Text/STextBlock.h"
+
+TMap<FName, FString> FSlateInspectorSnapshotRenderer::TypeToRole;
+TArray<TPair<FString, FString>> FSlateInspectorSnapshotRenderer::TemplatePrefixToRole;
+TMap<FString, FString> FSlateInspectorSnapshotRenderer::RoleToRefPrefix;
+TSet<FName> FSlateInspectorSnapshotRenderer::StructuralContainers;
+TMap<FName, TFunction<FString(TSharedRef<SWidget>)>> FSlateInspectorSnapshotRenderer::CustomLabelExtractors;
+TMap<FName, TFunction<TArray<FString>(TSharedRef<SWidget>)>> FSlateInspectorSnapshotRenderer::CustomStateFlagsExtractors;
+bool FSlateInspectorSnapshotRenderer::bMapsInitialized = false;
+
+void FSlateInspectorSnapshotRenderer::EnsureMapsInitialized()
+{
+	if (bMapsInitialized)
+	{
+		return;
+	}
+	bMapsInitialized = true;
+
+	// Type -> role mapping (non-template types use exact FName match).
+	TypeToRole.Add(FName(TEXT("SWindow")),               TEXT("window"));
+	TypeToRole.Add(FName(TEXT("SButton")),               TEXT("button"));
+	TypeToRole.Add(FName(TEXT("SEditableText")),          TEXT("textbox"));
+	TypeToRole.Add(FName(TEXT("SEditableTextBox")),       TEXT("textbox"));
+	TypeToRole.Add(FName(TEXT("SMultiLineEditableText")), TEXT("textbox"));
+	TypeToRole.Add(FName(TEXT("SCheckBox")),              TEXT("checkbox"));
+	TypeToRole.Add(FName(TEXT("SComboButton")),           TEXT("combobox"));
+	TypeToRole.Add(FName(TEXT("SSlider")),                TEXT("slider"));
+	TypeToRole.Add(FName(TEXT("SMenuBar")),               TEXT("menubar"));
+	TypeToRole.Add(FName(TEXT("SMenuAnchor")),            TEXT("menu"));
+	TypeToRole.Add(FName(TEXT("SScrollBox")),             TEXT("scrollable"));
+	TypeToRole.Add(FName(TEXT("SSplitter")),              TEXT("splitter"));
+	TypeToRole.Add(FName(TEXT("SDockTab")),               TEXT("tab"));
+	TypeToRole.Add(FName(TEXT("STabWell")),               TEXT("tablist"));
+	TypeToRole.Add(FName(TEXT("SToolTip")),               TEXT("tooltip"));
+	TypeToRole.Add(FName(TEXT("STextBlock")),             TEXT("text"));
+	TypeToRole.Add(FName(TEXT("SRichTextBlock")),         TEXT("text"));
+	TypeToRole.Add(FName(TEXT("SImage")),                 TEXT("image"));
+	TypeToRole.Add(FName(TEXT("SProgressBar")),           TEXT("progressbar"));
+
+	// Template types use prefix matching (e.g., "SComboBox<" matches "SComboBox<TSharedPtr<FString>>").
+	TemplatePrefixToRole.Add(TPair<FString, FString>(TEXT("SComboBox<"),  TEXT("combobox")));
+	TemplatePrefixToRole.Add(TPair<FString, FString>(TEXT("SSpinBox<"),   TEXT("slider")));
+	TemplatePrefixToRole.Add(TPair<FString, FString>(TEXT("STreeView<"),  TEXT("treeview")));
+	TemplatePrefixToRole.Add(TPair<FString, FString>(TEXT("SListView<"),  TEXT("listview")));
+	TemplatePrefixToRole.Add(TPair<FString, FString>(TEXT("STableRow<"),  TEXT("listitem")));
+
+	// Role -> ref prefix mapping.
+	RoleToRefPrefix.Add(TEXT("window"),      TEXT("w"));
+	RoleToRefPrefix.Add(TEXT("button"),      TEXT("b"));
+	RoleToRefPrefix.Add(TEXT("textbox"),     TEXT("tb"));
+	RoleToRefPrefix.Add(TEXT("checkbox"),    TEXT("cb"));
+	RoleToRefPrefix.Add(TEXT("combobox"),    TEXT("co"));
+	RoleToRefPrefix.Add(TEXT("slider"),      TEXT("sl"));
+	RoleToRefPrefix.Add(TEXT("menubar"),     TEXT("mb"));
+	RoleToRefPrefix.Add(TEXT("menu"),        TEXT("m"));
+	RoleToRefPrefix.Add(TEXT("treeview"),    TEXT("tv"));
+	RoleToRefPrefix.Add(TEXT("listview"),    TEXT("lv"));
+	RoleToRefPrefix.Add(TEXT("listitem"),    TEXT("li"));
+	RoleToRefPrefix.Add(TEXT("scrollable"),  TEXT("sc"));
+	RoleToRefPrefix.Add(TEXT("splitter"),    TEXT("sp"));
+	RoleToRefPrefix.Add(TEXT("tab"),         TEXT("t"));
+	RoleToRefPrefix.Add(TEXT("tablist"),     TEXT("tl"));
+	RoleToRefPrefix.Add(TEXT("tooltip"),     TEXT("tt"));
+	RoleToRefPrefix.Add(TEXT("text"),        TEXT("x"));
+	RoleToRefPrefix.Add(TEXT("image"),       TEXT("i"));
+	RoleToRefPrefix.Add(TEXT("progressbar"), TEXT("pb"));
+
+	// Structural containers to skip (promote children).
+	StructuralContainers.Add(FName(TEXT("SBox")));
+	StructuralContainers.Add(FName(TEXT("SBorder")));
+	StructuralContainers.Add(FName(TEXT("SOverlay")));
+	StructuralContainers.Add(FName(TEXT("SSpacer")));
+	StructuralContainers.Add(FName(TEXT("SConstraintCanvas")));
+	StructuralContainers.Add(FName(TEXT("SHorizontalBox")));
+	StructuralContainers.Add(FName(TEXT("SVerticalBox")));
+	StructuralContainers.Add(FName(TEXT("SGridPanel")));
+	StructuralContainers.Add(FName(TEXT("SWrapBox")));
+	StructuralContainers.Add(FName(TEXT("SWidgetSwitcher")));
+	StructuralContainers.Add(FName(TEXT("SCanvas")));
+	StructuralContainers.Add(FName(TEXT("SScaleBox")));
+	StructuralContainers.Add(FName(TEXT("SSizeBox")));
+	StructuralContainers.Add(FName(TEXT("SNullWidget")));
+	StructuralContainers.Add(FName(TEXT("SInvalidationPanel")));
+	StructuralContainers.Add(FName(TEXT("SRetainerWidget")));
+}
+
+void FSlateInspectorSnapshotRenderer::RegisterWidgetRole(const FName& WidgetType, const FString& Role, const FString& RefPrefix)
+{
+	EnsureMapsInitialized();
+	TypeToRole.Add(WidgetType, Role);
+	RoleToRefPrefix.Add(Role, RefPrefix);
+}
+
+void FSlateInspectorSnapshotRenderer::RegisterLabelExtractor(const FName& WidgetType, TFunction<FString(TSharedRef<SWidget>)> Extractor)
+{
+	CustomLabelExtractors.Add(WidgetType, MoveTemp(Extractor));
+}
+
+void FSlateInspectorSnapshotRenderer::RegisterStateFlagsExtractor(const FName& WidgetType, TFunction<TArray<FString>(TSharedRef<SWidget>)> Extractor)
+{
+	CustomStateFlagsExtractors.Add(WidgetType, MoveTemp(Extractor));
+}
+
+const FString* FSlateInspectorSnapshotRenderer::FindRoleForType(const FName& WidgetType, const FString& TypeString)
+{
+	// Try exact match first.
+	if (const FString* Role = TypeToRole.Find(WidgetType))
+	{
+		return Role;
+	}
+
+	// Try prefix match for template types (e.g., "SComboBox<TSharedPtr<FString>>").
+	for (const TPair<FString, FString>& Pair : TemplatePrefixToRole)
+	{
+		if (TypeString.StartsWith(Pair.Key))
+		{
+			return &Pair.Value;
+		}
+	}
+
+	return nullptr;
+}
+
+FString FSlateInspectorSnapshotRenderer::Render(TSharedPtr<SWidget> Root, int32 MaxDepth, bool bIncludeSourceLocations)
+{
+	EnsureMapsInitialized();
+
+	FSlateInspectorRefCache::Get().Reset();
+
+	TStringBuilder<4096> Output;
+
+	if (Root.IsValid())
+	{
+		RenderWidget(Output, Root.ToSharedRef(), 0, 0, MaxDepth, bIncludeSourceLocations);
+	}
+	else
+	{
+		// Snapshot all visible windows in z-order.
+		if (FSlateApplication::IsInitialized())
+		{
+			TArray<TSharedRef<SWindow>> VisibleWindows;
+			FSlateApplication::Get().GetAllVisibleWindowsOrdered(VisibleWindows);
+			for (const TSharedRef<SWindow>& Window : VisibleWindows)
+			{
+				RenderWidget(Output, Window, 0, 0, MaxDepth, bIncludeSourceLocations);
+			}
+		}
+	}
+
+	return Output.ToString();
+}
+
+void FSlateInspectorSnapshotRenderer::RenderWidget(
+	FStringBuilderBase& Output,
+	TSharedRef<SWidget> Widget,
+	int32 IndentLevel,
+	int32 CurrentDepth,
+	int32 MaxDepth,
+	bool bIncludeSourceLocations)
+{
+	if (CurrentDepth > MaxDepth)
+	{
+		return;
+	}
+
+	// Skip collapsed/hidden widgets.
+	if (Widget->GetVisibility() == EVisibility::Collapsed || Widget->GetVisibility() == EVisibility::Hidden)
+	{
+		return;
+	}
+
+	const FName WidgetType = Widget->GetType();
+	const FString TypeString = Widget->GetTypeAsString();
+
+	// Check if this is a structural container (skip but recurse children at same indent).
+	if (StructuralContainers.Contains(WidgetType))
+	{
+		FChildren* Children = Widget->GetChildren();
+		for (int32 ChildIndex = 0; ChildIndex < Children->Num(); ++ChildIndex)
+		{
+			RenderWidget(Output, Children->GetChildAt(ChildIndex), IndentLevel, CurrentDepth + 1, MaxDepth, bIncludeSourceLocations);
+		}
+		return;
+	}
+
+	// Look up role for this widget type (handles both exact and template prefix match).
+	const FString* Role = FindRoleForType(WidgetType, TypeString);
+
+	// Text collapsing: if this widget has exactly one child and that child is STextBlock,
+	// absorb the text child's label and don't recurse into it.
+	bool bAbsorbedTextChild = false;
+	FString AbsorbedText;
+	FChildren* Children = Widget->GetChildren();
+	if (Children->Num() == 1)
+	{
+		TSharedRef<SWidget> OnlyChild = Children->GetChildAt(0);
+		if (OnlyChild->GetType() == FName(TEXT("STextBlock")))
+		{
+			TSharedRef<STextBlock> TextBlock = StaticCastSharedRef<STextBlock>(OnlyChild);
+			AbsorbedText = TextBlock->GetText().ToString();
+			bAbsorbedTextChild = true;
+		}
+	}
+
+	// If we have no role and no useful label, skip this widget but recurse children.
+	if (!Role)
+	{
+		FString Label = ExtractLabel(Widget, TypeString);
+		if (Label.IsEmpty() && !bAbsorbedTextChild)
+		{
+			for (int32 ChildIndex = 0; ChildIndex < Children->Num(); ++ChildIndex)
+			{
+				RenderWidget(Output, Children->GetChildAt(ChildIndex), IndentLevel, CurrentDepth + 1, MaxDepth, bIncludeSourceLocations);
+			}
+			return;
+		}
+	}
+
+	// This widget is meaningful — emit it.
+	FString Label;
+	if (bAbsorbedTextChild && !AbsorbedText.IsEmpty())
+	{
+		Label = AbsorbedText;
+	}
+	else
+	{
+		Label = ExtractLabel(Widget, TypeString);
+	}
+
+	FString StateFlags = ExtractStateFlags(Widget, TypeString);
+
+	// Determine the role string and ref prefix.
+	FString RoleString = Role ? *Role : TEXT("generic");
+	const FString* RefPrefix = RoleToRefPrefix.Find(RoleString);
+	FString Prefix = RefPrefix ? *RefPrefix : TEXT("g");
+
+	FString Ref = FSlateInspectorRefCache::Get().GetOrAssignRef(Widget, Prefix);
+
+	// Emit the line.
+	for (int32 Index = 0; Index < IndentLevel; ++Index)
+	{
+		Output.Append(TEXT("  "));
+	}
+
+	Output.Append(*RoleString);
+
+	if (!Label.IsEmpty())
+	{
+		Output.Append(TEXT(" \""));
+		// Escape quotes and control characters that would break the line-oriented format.
+		for (TCHAR Character : Label)
+		{
+			switch (Character)
+			{
+			case TEXT('"'):  Output.Append(TEXT("\\\"")); break;
+			case TEXT('\\'): Output.Append(TEXT("\\\\")); break;
+			case TEXT('\n'): Output.Append(TEXT("\\n"));  break;
+			case TEXT('\r'): Output.Append(TEXT("\\r"));  break;
+			case TEXT('\t'): Output.Append(TEXT("\\t"));  break;
+			default:         Output.AppendChar(Character); break;
+			}
+		}
+		Output.Append(TEXT("\""));
+	}
+
+	if (!StateFlags.IsEmpty())
+	{
+		Output.Append(TEXT(" "));
+		Output.Append(*StateFlags);
+	}
+
+	// Geometry: absolute position (desktop space) and local size (widget space).
+	const FGeometry& Geometry = Widget->GetCachedGeometry();
+	const FVector2D AbsolutePosition = Geometry.GetAbsolutePosition();
+	const FVector2D LocalSize = Geometry.GetLocalSize();
+	Output.Appendf(TEXT(" [pos=%.0f,%.0f size=%.0f,%.0f]"), AbsolutePosition.X, AbsolutePosition.Y, LocalSize.X, LocalSize.Y);
+
+	if (bIncludeSourceLocations)
+	{
+#if !UE_BUILD_SHIPPING
+		FName CreatedInLocation = Widget->GetCreatedInLocation();
+		if (!CreatedInLocation.IsNone())
+		{
+			FString FilePath = CreatedInLocation.GetPlainNameString();
+			int32 LineNumber = CreatedInLocation.GetNumber();
+			FString FileName = FPaths::GetCleanFilename(FilePath);
+			if (LineNumber > 0)
+			{
+				Output.Appendf(TEXT(" [src=%s:%d]"), *FileName, LineNumber);
+			}
+			else
+			{
+				Output.Appendf(TEXT(" [src=%s]"), *FileName);
+			}
+		}
+#endif
+	}
+
+	Output.Append(TEXT(" [ref="));
+	Output.Append(*Ref);
+	Output.Append(TEXT("]\n"));
+
+	// Recurse into children (unless we absorbed the only text child).
+	if (!bAbsorbedTextChild)
+	{
+		for (int32 ChildIndex = 0; ChildIndex < Children->Num(); ++ChildIndex)
+		{
+			RenderWidget(Output, Children->GetChildAt(ChildIndex), IndentLevel + 1, CurrentDepth + 1, MaxDepth, bIncludeSourceLocations);
+		}
+	}
+}
+
+FString FSlateInspectorSnapshotRenderer::ExtractLabel(TSharedRef<SWidget> Widget, const FString& TypeString)
+{
+	// Prefer Slate's native accessible text — many widgets implement this already.
+#if WITH_ACCESSIBILITY
+	FText AccessibleText = Widget->GetAccessibleText(EAccessibleType::Main);
+	if (!AccessibleText.IsEmpty())
+	{
+		return AccessibleText.ToString();
+	}
+#endif
+
+	// Fall back to type-specific extraction when no accessible text is available.
+	if (TypeString == TEXT("SWindow"))
+	{
+		return StaticCastSharedRef<SWindow>(Widget)->GetTitle().ToString();
+	}
+	if (TypeString == TEXT("STextBlock"))
+	{
+		return StaticCastSharedRef<STextBlock>(Widget)->GetText().ToString();
+	}
+	if (TypeString == TEXT("SEditableTextBox"))
+	{
+		return StaticCastSharedRef<SEditableTextBox>(Widget)->GetText().ToString();
+	}
+	if (TypeString == TEXT("SDockTab"))
+	{
+		return StaticCastSharedRef<SDockTab>(Widget)->GetTabLabel().ToString();
+	}
+
+	// Try custom label extractors.
+	if (const auto* Extractor = CustomLabelExtractors.Find(Widget->GetType()))
+	{
+		return (*Extractor)(Widget);
+	}
+
+	return FString();
+}
+
+FString FSlateInspectorSnapshotRenderer::ExtractStateFlags(TSharedRef<SWidget> Widget, const FString& TypeString)
+{
+	TArray<FString> Flags;
+
+	if (!Widget->IsEnabled())
+	{
+		Flags.Add(TEXT("disabled"));
+	}
+
+	if (Widget->HasKeyboardFocus())
+	{
+		Flags.Add(TEXT("focused"));
+	}
+
+	if (TypeString == TEXT("SCheckBox"))
+	{
+		ECheckBoxState CheckedState = StaticCastSharedRef<SCheckBox>(Widget)->GetCheckedState();
+		switch (CheckedState)
+		{
+		case ECheckBoxState::Checked:
+			Flags.Add(TEXT("checked"));
+			break;
+		case ECheckBoxState::Undetermined:
+			Flags.Add(TEXT("indeterminate"));
+			break;
+		default:
+			Flags.Add(TEXT("unchecked"));
+			break;
+		}
+	}
+
+	// Try custom state flags extractors.
+	if (const auto* Extractor = CustomStateFlagsExtractors.Find(Widget->GetType()))
+	{
+		Flags.Append((*Extractor)(Widget));
+	}
+
+	if (Flags.IsEmpty())
+	{
+		return FString();
+	}
+
+	return FString::Printf(TEXT("[%s]"), *FString::Join(Flags, TEXT(", ")));
+}
diff --git a/Engine/Plugins/Experimental/SlateInspector/Source/SlateInspector/Private/SlateInspectorSubsystem.cpp b/Engine/Plugins/Experimental/SlateInspector/Source/SlateInspector/Private/SlateInspectorSubsystem.cpp
new file mode 100644
--- /dev/null
+++ b/Engine/Plugins/Experimental/SlateInspector/Source/SlateInspector/Private/SlateInspectorSubsystem.cpp
@@ -0,0 +1,56 @@
+// Copyright Epic Games, Inc. All Rights Reserved.
+
+#include "SlateInspectorSubsystem.h"
+
+#include "Editor.h"
+#include "ToolsetRegistry/UToolsetRegistry.h"
+#include "SlateInspectorToolset.h"
+#include "Subsystems/SubsystemCollection.h"
+
+static bool bEnableSlateInspector = true;
+
+static void OnEnableSlateInspectorChanged(IConsoleVariable* Variable)
+{
+	if (GEditor == nullptr)
+	{
+		return;
+	}
+	
+	if (USlateInspectorSubsystem* Subsystem = GEditor->GetEditorSubsystem<USlateInspectorSubsystem>())
+	{
+		Subsystem->SetToolsetEnabled(Variable->GetBool());
+	}
+}
+
+static FAutoConsoleVariableRef CVarEnableSlateInspector(
+	TEXT("SlateInspectorToolset.Enable"),
+	bEnableSlateInspector,
+	TEXT("Enable or disable the SlateInspector toolset registration. When disabled, SlateInspector tools will not appear in the MCP tool list."),
+	FConsoleVariableDelegate::CreateStatic(&OnEnableSlateInspectorChanged),
+	ECVF_Default);
+
+void USlateInspectorSubsystem::Initialize(FSubsystemCollectionBase& Collection)
+{
+	Super::Initialize(Collection);
+	SetToolsetEnabled(bEnableSlateInspector);
+}
+
+void USlateInspectorSubsystem::Deinitialize()
+{
+	SetToolsetEnabled(false);
+	Super::Deinitialize();
+}
+
+void USlateInspectorSubsystem::SetToolsetEnabled(bool bEnabled)
+{
+	if (bEnabled && !bToolsetRegistered)
+	{
+		UToolsetRegistry::RegisterToolsetClass(USlateInspectorToolset::StaticClass());
+		bToolsetRegistered = true;
+	}
+	else if (!bEnabled && bToolsetRegistered)
+	{
+		UToolsetRegistry::UnregisterToolsetClass(USlateInspectorToolset::StaticClass());
+		bToolsetRegistered = false;
+	}
+}
diff --git a/Engine/Plugins/Experimental/SlateInspector/Source/SlateInspector/Private/SlateInspectorToolset.cpp b/Engine/Plugins/Experimental/SlateInspector/Source/SlateInspector/Private/SlateInspectorToolset.cpp
new file mode 100644
--- /dev/null
+++ b/Engine/Plugins/Experimental/SlateInspector/Source/SlateInspector/Private/SlateInspectorToolset.cpp
@@ -0,0 +1,573 @@
+// Copyright Epic Games, Inc. All Rights Reserved.
+
+#include "SlateInspectorToolset.h"
+
+#include "SlateInspectorRefCache.h"
+#include "SlateInspectorSnapshotRenderer.h"
+
+#include "Framework/Application/SlateApplication.h"
+#include "Input/Events.h"
+#include "Logging/StructuredLog.h"
+#include "Widgets/Input/SCheckBox.h"
+#include "Widgets/SWindow.h"
+#include "Widgets/Text/STextBlock.h"
+
+DEFINE_LOG_CATEGORY_STATIC(LogSlateInspector, Log, All);
+
+namespace SlateInspectorPrivate
+{
+
+/** Escape a string for safe embedding in a JSON string value (RFC 8259). */
+FString EscapeJsonString(const FString& Input)
+{
+	FString Result;
+	Result.Reserve(Input.Len() + 16);
+	for (TCHAR Character : Input)
+	{
+		switch (Character)
+		{
+		case TEXT('"'):  Result += TEXT("\\\""); break;
+		case TEXT('\\'): Result += TEXT("\\\\"); break;
+		case TEXT('\b'): Result += TEXT("\\b");  break;
+		case TEXT('\f'): Result += TEXT("\\f");  break;
+		case TEXT('\n'): Result += TEXT("\\n");  break;
+		case TEXT('\r'): Result += TEXT("\\r");  break;
+		case TEXT('\t'): Result += TEXT("\\t");  break;
+		default:
+			// Escape remaining ASCII control characters (U+0000..U+001F) per RFC 8259.
+			if (Character < 0x20)
+			{
+				Result += FString::Printf(TEXT("\\u%04x"), static_cast<uint32>(Character));
+			}
+			else
+			{
+				Result += Character;
+			}
+			break;
+		}
+	}
+	return Result;
+}
+
+} // namespace SlateInspectorPrivate
+
+// ── Helpers ──────────────────────────────────────────────────────────
+
+TSharedPtr<SWidget> USlateInspectorToolset::ResolveRefOrWarn(const FString& Ref)
+{
+	TSharedPtr<SWidget> Widget = FSlateInspectorRefCache::Get().ResolveRef(Ref);
+	if (!Widget.IsValid())
+	{
+		UE_LOGFMT(LogSlateInspector, Warning, "SlateInspector: Could not resolve ref '{Ref}'", Ref);
+	}
+	return Widget;
+}
+
+FVector2D USlateInspectorToolset::GetWidgetScreenCenter(TSharedRef<SWidget> Widget)
+{
+	const FGeometry& Geometry = Widget->GetTickSpaceGeometry();
+	return FVector2D(Geometry.GetAbsolutePositionAtCoordinates(FVector2D(0.5, 0.5)));
+}
+
+FKey USlateInspectorToolset::ParseMouseButtonKey(const FString& Button)
+{
+	if (Button.Equals(TEXT("right"), ESearchCase::IgnoreCase))
+	{
+		return EKeys::RightMouseButton;
+	}
+	if (Button.Equals(TEXT("middle"), ESearchCase::IgnoreCase))
+	{
+		return EKeys::MiddleMouseButton;
+	}
+	return EKeys::LeftMouseButton;
+}
+
+bool USlateInspectorToolset::SimulateClick(TSharedRef<SWidget> Widget, FKey MouseButton, bool bDoubleClick)
+{
+	if (!FSlateApplication::IsInitialized())
+	{
+		return false;
+	}
+
+	TSharedPtr<SWindow> Window = FSlateApplication::Get().FindWidgetWindow(Widget);
+	if (!Window.IsValid())
+	{
+		UE_LOGFMT(LogSlateInspector, Warning, "SlateInspector: Could not find window for widget");
+		return false;
+	}
+
+	FVector2D ScreenCenter = GetWidgetScreenCenter(Widget);
+	FSlateApplication::Get().SetCursorPos(ScreenCenter);
+
+	FModifierKeysState ModifierKeys;
+	TSet<FKey> PressedButtons;
+	PressedButtons.Add(MouseButton);
+
+	FPointerEvent MouseEvent(
+		0,
+		ScreenCenter,
+		ScreenCenter,
+		PressedButtons,
+		MouseButton,
+		0.0f,
+		ModifierKeys);
+
+	if (bDoubleClick)
+	{
+		FSlateApplication::Get().ProcessMouseButtonDoubleClickEvent(Window->GetNativeWindow(), MouseEvent);
+	}
+	else
+	{
+		FSlateApplication::Get().ProcessMouseButtonDownEvent(Window->GetNativeWindow(), MouseEvent);
+
+		// Release with empty pressed buttons set.
+		FPointerEvent MouseUpEvent(
+			0,
+			ScreenCenter,
+			ScreenCenter,
+			TSet<FKey>(),
+			MouseButton,
+			0.0f,
+			ModifierKeys);
+		FSlateApplication::Get().ProcessMouseButtonUpEvent(MouseUpEvent);
+	}
+
+	return true;
+}
+
+// ── Snapshot tools ───────────────────────────────────────────────────
+
+FString USlateInspectorToolset::Snapshot(const FString& Ref, int32 MaxDepth, bool bIncludeSourceLocations)
+{
+	TSharedPtr<SWidget> Root;
+	if (!Ref.IsEmpty())
+	{
+		Root = ResolveRefOrWarn(Ref);
+		if (!Root.IsValid())
+		{
+			return TEXT("Error: Could not resolve ref");
+		}
+	}
+
+	return FSlateInspectorSnapshotRenderer::Render(Root, MaxDepth, bIncludeSourceLocations);
+}
+
+FToolsetImage USlateInspectorToolset::Screenshot(const FString& Ref)
+{
+	FToolsetImage Result;
+
+	if (!FSlateApplication::IsInitialized())
+	{
+		return Result;
+	}
+
+	TSharedPtr<SWidget> Widget;
+	if (!Ref.IsEmpty())
+	{
+		Widget = ResolveRefOrWarn(Ref);
+	}
+	else
+	{
+		// Use the active top-level window.
+		TSharedPtr<SWindow> ActiveWindow = FSlateApplication::Get().GetActiveTopLevelWindow();
+		Widget = ActiveWindow;
+	}
+
+	if (!Widget.IsValid())
+	{
+		return Result;
+	}
+
+	TArray<FColor> ColorData;
+	FIntVector Size;
+	if (FSlateApplication::Get().TakeScreenshot(Widget.ToSharedRef(), ColorData, Size))
+	{
+		Result.SetFromBitmap(ColorData, FIntPoint(Size.X, Size.Y));
+	}
+
+	return Result;
+}
+
+// ── Action tools ─────────────────────────────────────────────────────
+
+bool USlateInspectorToolset::Click(const FString& Ref, const FString& Button, bool DoubleClick)
+{
+	TSharedPtr<SWidget> Widget = ResolveRefOrWarn(Ref);
+	if (!Widget.IsValid())
+	{
+		return false;
+	}
+
+	return SimulateClick(Widget.ToSharedRef(), ParseMouseButtonKey(Button), DoubleClick);
+}
+
+bool USlateInspectorToolset::Hover(const FString& Ref)
+{
+	TSharedPtr<SWidget> Widget = ResolveRefOrWarn(Ref);
+	if (!Widget.IsValid() || !FSlateApplication::IsInitialized())
+	{
+		return false;
+	}
+
+	FVector2D ScreenCenter = GetWidgetScreenCenter(Widget.ToSharedRef());
+	FSlateApplication::Get().SetCursorPos(ScreenCenter);
+	return true;
+}
+
+bool USlateInspectorToolset::Type(const FString& Ref, const FString& Text, bool Submit)
+{
+	TSharedPtr<SWidget> Widget = ResolveRefOrWarn(Ref);
+	if (!Widget.IsValid() || !FSlateApplication::IsInitialized())
+	{
+		return false;
+	}
+
+	// Focus the target widget.
+	FSlateApplication::Get().SetKeyboardFocus(Widget, EFocusCause::SetDirectly);
+
+	FModifierKeysState ModifierKeys;
+
+	// Send each character as a key char event.
+	for (TCHAR Character : Text)
+	{
+		FCharacterEvent CharEvent(Character, ModifierKeys, 0, false);
+		FSlateApplication::Get().ProcessKeyCharEvent(CharEvent);
+	}
+
+	if (Submit)
+	{
+		FKeyEvent KeyDownEvent(EKeys::Enter, ModifierKeys, 0, false, 0, 0);
+		FSlateApplication::Get().ProcessKeyDownEvent(KeyDownEvent);
+		FKeyEvent KeyUpEvent(EKeys::Enter, ModifierKeys, 0, false, 0, 0);
+		FSlateApplication::Get().ProcessKeyUpEvent(KeyUpEvent);
+	}
+
+	return true;
+}
+
+bool USlateInspectorToolset::PressKey(const FString& Key)
+{
+	if (!FSlateApplication::IsInitialized())
+	{
+		return false;
+	}
+
+	FName KeyName(*Key);
+	FKey ParsedKey(KeyName);
+
+	if (!ParsedKey.IsValid())
+	{
+		UE_LOGFMT(LogSlateInspector, Warning, "SlateInspector: Unknown key '{Key}'", Key);
+		return false;
+	}
+
+	FModifierKeysState ModifierKeys;
+	FKeyEvent KeyDownEvent(ParsedKey, ModifierKeys, 0, false, 0, 0);
+	FSlateApplication::Get().ProcessKeyDownEvent(KeyDownEvent);
+	FKeyEvent KeyUpEvent(ParsedKey, ModifierKeys, 0, false, 0, 0);
+	FSlateApplication::Get().ProcessKeyUpEvent(KeyUpEvent);
+
+	return true;
+}
+
+bool USlateInspectorToolset::SelectOption(const FString& Ref, const FString& Value)
+{
+	TSharedPtr<SWidget> Widget = ResolveRefOrWarn(Ref);
+	if (!Widget.IsValid())
+	{
+		return false;
+	}
+
+	// Click to open the combobox dropdown.
+	if (!SimulateClick(Widget.ToSharedRef(), EKeys::LeftMouseButton, false))
+	{
+		return false;
+	}
+
+	// Search all visible windows for a text block matching the value.
+	if (!FSlateApplication::IsInitialized())
+	{
+		return false;
+	}
+
+	TSharedPtr<SWidget> MatchingTextBlock;
+	TArray<TSharedRef<SWindow>> VisibleWindows;
+	FSlateApplication::Get().GetAllVisibleWindowsOrdered(VisibleWindows);
+	for (const TSharedRef<SWindow>& Window : VisibleWindows)
+	{
+		TArray<TSharedRef<SWidget>> Queue;
+		FChildren* WindowChildren = Window->GetChildren();
+		for (int32 ChildIndex = 0; ChildIndex < WindowChildren->Num(); ++ChildIndex)
+		{
+			Queue.Add(WindowChildren->GetChildAt(ChildIndex));
+		}
+
+		static const FName TextBlockType(TEXT("STextBlock"));
+		while (Queue.Num() > 0)
+		{
+			TSharedRef<SWidget> Current = Queue[0];
+			Queue.RemoveAt(0);
+
+			if (Current->GetType() == TextBlockType)
+			{
+				FText WidgetText = StaticCastSharedRef<STextBlock>(Current)->GetText();
+				if (WidgetText.ToString() == Value)
+				{
+					MatchingTextBlock = Current;
+					break;
+				}
+			}
+
+			FChildren* CurrentChildren = Current->GetChildren();
+			for (int32 ChildIndex = 0; ChildIndex < CurrentChildren->Num(); ++ChildIndex)
+			{
+				Queue.Add(CurrentChildren->GetChildAt(ChildIndex));
+			}
+		}
+
+		if (MatchingTextBlock.IsValid())
+		{
+			break;
+		}
+	}
+
+	if (!MatchingTextBlock.IsValid())
+	{
+		UE_LOGFMT(LogSlateInspector, Warning, "SlateInspector: Could not find option '{Value}'", Value);
+		return false;
+	}
+
+	return SimulateClick(MatchingTextBlock.ToSharedRef(), EKeys::LeftMouseButton, false);
+}
+
+bool USlateInspectorToolset::Drag(const FString& StartRef, const FString& EndRef)
+{
+	TSharedPtr<SWidget> StartWidget = ResolveRefOrWarn(StartRef);
+	TSharedPtr<SWidget> EndWidget = ResolveRefOrWarn(EndRef);
+	if (!StartWidget.IsValid() || !EndWidget.IsValid() || !FSlateApplication::IsInitialized())
+	{
+		return false;
+	}
+
+	TSharedPtr<SWindow> StartWindow = FSlateApplication::Get().FindWidgetWindow(StartWidget.ToSharedRef());
+	if (!StartWindow.IsValid())
+	{
+		UE_LOGFMT(LogSlateInspector, Warning, "SlateInspector: Could not find window for drag source");
+		return false;
+	}
+
+	FVector2D StartPosition = GetWidgetScreenCenter(StartWidget.ToSharedRef());
+	FVector2D EndPosition = GetWidgetScreenCenter(EndWidget.ToSharedRef());
+
+	FModifierKeysState ModifierKeys;
+	TSet<FKey> PressedButtons;
+	PressedButtons.Add(EKeys::LeftMouseButton);
+
+	// Move cursor to start and press.
+	FSlateApplication::Get().SetCursorPos(StartPosition);
+	FPointerEvent MouseDownEvent(
+		0, StartPosition, StartPosition,
+		PressedButtons, EKeys::LeftMouseButton, 0.0f, ModifierKeys);
+	FSlateApplication::Get().ProcessMouseButtonDownEvent(StartWindow->GetNativeWindow(), MouseDownEvent);
+
+	// Move cursor to end.
+	FPointerEvent MouseMoveEvent(
+		0, EndPosition, StartPosition,
+		PressedButtons, EKeys::Invalid, 0.0f, ModifierKeys);
+	FSlateApplication::Get().ProcessMouseMoveEvent(MouseMoveEvent);
+	FSlateApplication::Get().SetCursorPos(EndPosition);
+
+	// Release at end position.
+	FPointerEvent MouseUpEvent(
+		0, EndPosition, EndPosition,
+		TSet<FKey>(), EKeys::LeftMouseButton, 0.0f, ModifierKeys);
+	FSlateApplication::Get().ProcessMouseButtonUpEvent(MouseUpEvent);
+
+	return true;
+}
+
+// ── Window tools ─────────────────────────────────────────────────────
+
+FString USlateInspectorToolset::Windows(const FString& Action, int32 Index)
+{
+	if (!FSlateApplication::IsInitialized())
+	{
+		return TEXT("Error: Slate not initialized");
+	}
+
+	TArray<TSharedRef<SWindow>> VisibleWindows;
+	FSlateApplication::Get().GetAllVisibleWindowsOrdered(VisibleWindows);
+
+	if (Action.Equals(TEXT("list"), ESearchCase::IgnoreCase))
+	{
+		FString Result = TEXT("[");
+		for (int32 WindowIndex = 0; WindowIndex < VisibleWindows.Num(); ++WindowIndex)
+		{
+			const TSharedRef<SWindow>& Window = VisibleWindows[WindowIndex];
+			if (WindowIndex > 0)
+			{
+				Result += TEXT(", ");
+			}
+			FString Title = SlateInspectorPrivate::EscapeJsonString(Window->GetTitle().ToString());
+			Result += FString::Printf(TEXT("{\"index\": %d, \"title\": \"%s\"}"),
+				WindowIndex,
+				*Title);
+		}
+		Result += TEXT("]");
+		return Result;
+	}
+
+	if (Action.Equals(TEXT("select"), ESearchCase::IgnoreCase))
+	{
+		if (Index >= 0 && Index < VisibleWindows.Num())
+		{
+			VisibleWindows[Index]->BringToFront();
+			FSlateApplication::Get().SetAllUserFocus(VisibleWindows[Index], EFocusCause::SetDirectly);
+			return TEXT("OK");
+		}
+		return TEXT("Error: Invalid window index");
+	}
+
+	if (Action.Equals(TEXT("close"), ESearchCase::IgnoreCase))
+	{
+		if (Index >= 0 && Index < VisibleWindows.Num())
+		{
+			VisibleWindows[Index]->RequestDestroyWindow();
+			return TEXT("OK");
+		}
+		return TEXT("Error: Invalid window index");
+	}
+
+	return TEXT("Error: Unknown action. Use 'list', 'select', or 'close'.");
+}
+
+// ── Wait tool ────────────────────────────────────────────────────────
+
+bool USlateInspectorToolset::WaitFor(const FString& Text, const FString& TextGone)
+{
+	if (!FSlateApplication::IsInitialized())
+	{
+		return false;
+	}
+
+	const bool bWaitForAppear = !Text.IsEmpty();
+	const bool bWaitForDisappear = !TextGone.IsEmpty();
+
+	if (!bWaitForAppear && !bWaitForDisappear)
+	{
+		return true;
+	}
+
+	// Walk the widget tree looking for text WITHOUT rendering a snapshot,
+	// because Render() resets the ref cache and would silently invalidate
+	// refs obtained from a prior Snapshot() call.
+	bool bFoundText = false;
+	bool bFoundTextGone = false;
+
+	TArray<TSharedRef<SWindow>> Windows;
+	FSlateApplication::Get().GetAllVisibleWindowsOrdered(Windows);
+
+	TFunction<void(TSharedRef<SWidget>)> Walk = [&](TSharedRef<SWidget> Widget)
+	{
+		// Early out if both conditions already determined.
+		if ((!bWaitForAppear || bFoundText) && (!bWaitForDisappear || bFoundTextGone))
+		{
+			return;
+		}
+
+		FString WidgetText;
+		if (Widget->GetType() == FName(TEXT("STextBlock")))
+		{
+			WidgetText = StaticCastSharedRef<STextBlock>(Widget)->GetText().ToString();
+		}
+#if WITH_ACCESSIBILITY
+		if (WidgetText.IsEmpty())
+		{
+			WidgetText = Widget->GetAccessibleText(EAccessibleType::Main).ToString();
+		}
+#endif
+
+		if (!WidgetText.IsEmpty())
+		{
+			if (bWaitForAppear && WidgetText.Contains(Text))
+			{
+				bFoundText = true;
+			}
+			if (bWaitForDisappear && WidgetText.Contains(TextGone))
+			{
+				bFoundTextGone = true;
+			}
+		}
+
+		FChildren* Children = Widget->GetChildren();
+		for (int32 Index = 0; Index < Children->Num(); ++Index)
+		{
+			Walk(Children->GetChildAt(Index));
+		}
+	};
+
+	for (const TSharedRef<SWindow>& Window : Windows)
+	{
+		Walk(Window);
+	}
+
+	bool bAppearConditionMet = !bWaitForAppear || bFoundText;
+	bool bDisappearConditionMet = !bWaitForDisappear || !bFoundTextGone;
+
+	return bAppearConditionMet && bDisappearConditionMet;
+}
+
+// ── Form tool ────────────────────────────────────────────────────────
+
+bool USlateInspectorToolset::FillForm(const TArray<FSlateInspectorFormField>& Fields)
+{
+	bool bAllSuccess = true;
+
+	for (const FSlateInspectorFormField& Field : Fields)
+	{
+		bool bFieldSuccess = false;
+
+		if (Field.FieldType.Equals(TEXT("textbox"), ESearchCase::IgnoreCase))
+		{
+			bFieldSuccess = Type(Field.Ref, Field.Value, false);
+		}
+		else if (Field.FieldType.Equals(TEXT("checkbox"), ESearchCase::IgnoreCase))
+		{
+			// Click the checkbox if the desired state differs from the current state.
+			TSharedPtr<SWidget> Widget = ResolveRefOrWarn(Field.Ref);
+			if (Widget.IsValid() && Widget->GetType() == FName(TEXT("SCheckBox")))
+			{
+				TSharedRef<SCheckBox> CheckBox = StaticCastSharedRef<SCheckBox>(Widget.ToSharedRef());
+				bool bDesiredChecked = Field.Value.Equals(TEXT("true"), ESearchCase::IgnoreCase);
+				bool bCurrentlyChecked = CheckBox->IsChecked();
+				if (bDesiredChecked != bCurrentlyChecked)
+				{
+					bFieldSuccess = Click(Field.Ref);
+				}
+				else
+				{
+					bFieldSuccess = true; // Already in desired state.
+				}
+			}
+		}
+		else if (Field.FieldType.Equals(TEXT("combobox"), ESearchCase::IgnoreCase))
+		{
+			bFieldSuccess = SelectOption(Field.Ref, Field.Value);
+		}
+		else if (Field.FieldType.Equals(TEXT("slider"), ESearchCase::IgnoreCase))
+		{
+			bFieldSuccess = Type(Field.Ref, Field.Value, true);
+		}
+		else
+		{
+			UE_LOGFMT(LogSlateInspector, Warning, "SlateInspector: Unknown field type '{FieldType}'", Field.FieldType);
+		}
+
+		if (!bFieldSuccess)
+		{
+			bAllSuccess = false;
+		}
+	}
+
+	return bAllSuccess;
+}
diff --git a/Engine/Plugins/Experimental/SlateInspector/Source/SlateInspector/Public/SlateInspectorModule.h b/Engine/Plugins/Experimental/SlateInspector/Source/SlateInspector/Public/SlateInspectorModule.h
new file mode 100644
--- /dev/null
+++ b/Engine/Plugins/Experimental/SlateInspector/Source/SlateInspector/Public/SlateInspectorModule.h
@@ -0,0 +1,14 @@
+// Copyright Epic Games, Inc. All Rights Reserved.
+
+#pragma once
+
+#include "Modules/ModuleInterface.h"
+
+class FSlateInspectorModule : public IModuleInterface
+{
+public:
+
+	virtual void StartupModule() override;
+
+	virtual void ShutdownModule() override;
+};
diff --git a/Engine/Plugins/Experimental/SlateInspector/Source/SlateInspector/Public/SlateInspectorRefCache.h b/Engine/Plugins/Experimental/SlateInspector/Source/SlateInspector/Public/SlateInspectorRefCache.h
new file mode 100644
--- /dev/null
+++ b/Engine/Plugins/Experimental/SlateInspector/Source/SlateInspector/Public/SlateInspectorRefCache.h
@@ -0,0 +1,89 @@
+// Copyright Epic Games, Inc. All Rights Reserved.
+
+#pragma once
+
+#include "CoreMinimal.h"
+#include "Widgets/SWidget.h"
+
+/**
+ * Manages bidirectional mapping between SWidget instances and short reference
+ * strings used by the SlateInspector snapshot and action tools.
+ *
+ * # Reset-per-snapshot model
+ *
+ * This cache follows the establised approach such as Playwright's accessibility snapshots:
+ * every call to FSlateInspectorSnapshotRenderer::Render() calls Reset() first,
+ * then the tree walk assigns fresh refs from scratch (b1, b2, tb1, ...).
+ * Counters restart at 1 each time, so refs are small, dense, and predictable.
+ *
+ * A ref is valid from the moment a snapshot is taken until the next snapshot
+ * replaces it.  The expected workflow is:
+ *
+ *   1. Take a snapshot  -> refs are assigned (b1, b2, tb1, ...)
+ *   2. Perform actions   -> Click("b2"), Type("tb1", "hello")
+ *   3. Take a new snapshot -> all previous refs are invalidated, new ones assigned
+ *
+ * This means callers must NOT cache refs across snapshots. Each snapshot is a
+ * self-contained view of the widget tree with its own ref namespace.
+ *
+ * # Rational
+ *
+ * - Compact refs: counters restart at 1 each pass, so refs stay small (b1 not b847).
+ * - Minimal stale refs: the full rebuild removes entries orphaned by prior
+ *   snapshots.  Between snapshots, destroyed widgets may leave expired weak
+ *   pointers — ResolveRef() returns nullptr for those.
+ * - Deterministic ordering: identical widget trees produce identical ref assignments,
+ *   making snapshots diff-friendly and easier for LLMs to reason about.
+ * - Simplicity: no pruning heuristics, no garbage collection, no ref-reuse logic.
+ *
+ * The trade-off is that refs are ephemeral: they only survive until the next
+ * snapshot. Expectation is that the caller always takes a fresh snapshot before acting,
+ * so refs never need long-term stability.
+ */
+class FSlateInspectorRefCache
+{
+public:
+
+	SLATEINSPECTOR_API static FSlateInspectorRefCache& Get();
+
+	/**
+	 * Returns existing ref or assigns a new one using the role prefix.
+	 *
+	 * The role prefix determines the ref format: "b" for buttons produces
+	 * b1, b2, b3, etc.  Within a single snapshot pass, the same widget
+	 * always gets the same ref (deduplication via WidgetToRef lookup).
+	 *
+	 * @param Widget      The widget to assign a ref to.
+	 * @param RolePrefix  Short prefix derived from the widget's role (e.g., "b", "tb", "cb").
+	 * @return The assigned ref string (e.g., "b3").
+	 */
+	SLATEINSPECTOR_API FString GetOrAssignRef(TSharedRef<SWidget> Widget, const FString& RolePrefix);
+
+	/**
+	 * Resolves a ref string back to a live widget.
+	 *
+	 * Returns nullptr if the ref is unknown or the widget has been destroyed
+	 * since the last snapshot.
+	 */
+	SLATEINSPECTOR_API TSharedPtr<SWidget> ResolveRef(const FString& Ref) const;
+
+	/**
+	 * Clears all mappings and resets counters to zero.
+	 *
+	 * Called at the start of each snapshot pass so that refs are assigned
+	 * fresh from scratch. After Reset(), the next button encountered will
+	 * be b1, the next textbox tb1, etc.
+	 */
+	SLATEINSPECTOR_API void Reset();
+
+private:
+
+	/** Widget -> ref string (deduplication within a single snapshot pass). */
+	TMap<TWeakPtr<SWidget>, FString> WidgetToRef;
+
+	/** Ref string -> widget (reverse lookup for action tools). */
+	TMap<FString, TWeakPtr<SWidget>> RefToWidget;
+
+	/** Per-role-prefix counters for generating sequential ref numbers. */
+	TMap<FString, int32> RolePrefixCounters;
+};
diff --git a/Engine/Plugins/Experimental/SlateInspector/Source/SlateInspector/Public/SlateInspectorSnapshotRenderer.h b/Engine/Plugins/Experimental/SlateInspector/Source/SlateInspector/Public/SlateInspectorSnapshotRenderer.h
new file mode 100644
--- /dev/null
+++ b/Engine/Plugins/Experimental/SlateInspector/Source/SlateInspector/Public/SlateInspectorSnapshotRenderer.h
@@ -0,0 +1,81 @@
+// Copyright Epic Games, Inc. All Rights Reserved.
+
+#pragma once
+
+#include "CoreMinimal.h"
+#include "Widgets/SWidget.h"
+
+/**
+ * Renders a text-based accessibility snapshot of a Slate widget tree.
+ *
+ * Walks the tree depth-first, emitting one line per meaningful widget with
+ * role, label, state flags, and a compact ref identifier.  Structural
+ * containers (SBox, SBorder, SOverlay, etc.) are elided — their children
+ * are emitted at the parent's indent level.
+ */
+class FSlateInspectorSnapshotRenderer
+{
+public:
+
+	/**
+	 * Render a snapshot starting from the given root widget.
+	 *
+	 * @param Root                     The widget subtree to snapshot.  If null, snapshots all top-level windows.
+	 * @param MaxDepth                 Maximum recursion depth (default 30).
+	 * @param bIncludeSourceLocations  Include [src=File:Line] tags showing where each widget was created in C++.
+	 * @return Indented text representation of the widget tree.
+	 */
+	SLATEINSPECTOR_API static FString Render(TSharedPtr<SWidget> Root, int32 MaxDepth = 30, bool bIncludeSourceLocations = false);
+
+	/** Register a custom widget type so it appears in snapshots. */
+	SLATEINSPECTOR_API static void RegisterWidgetRole(const FName& WidgetType, const FString& Role, const FString& RefPrefix);
+
+	/** Register a custom label extractor for a widget type. */
+	SLATEINSPECTOR_API static void RegisterLabelExtractor(const FName& WidgetType, TFunction<FString(TSharedRef<SWidget>)> Extractor);
+
+	/** Register a custom state flags extractor for a widget type. */
+	SLATEINSPECTOR_API static void RegisterStateFlagsExtractor(const FName& WidgetType, TFunction<TArray<FString>(TSharedRef<SWidget>)> Extractor);
+
+private:
+
+	/** Initializes the type-to-role and structural-skip maps on first use. */
+	static void EnsureMapsInitialized();
+
+	/** Recursive depth-first walk. */
+	static void RenderWidget(
+		FStringBuilderBase& Output,
+		TSharedRef<SWidget> Widget,
+		int32 IndentLevel,
+		int32 CurrentDepth,
+		int32 MaxDepth,
+		bool bIncludeSourceLocations);
+
+	/** Extract a human-readable label from a widget via type-specific methods. */
+	static FString ExtractLabel(TSharedRef<SWidget> Widget, const FString& TypeString);
+
+	/** Extract state flags (visibility, enabled, checked, focused). */
+	static FString ExtractStateFlags(TSharedRef<SWidget> Widget, const FString& TypeString);
+
+	/** Look up role for a widget type, handling template types via prefix matching. */
+	static const FString* FindRoleForType(const FName& WidgetType, const FString& TypeString);
+
+	/** Map from widget type name to role string (e.g., "SButton" -> "button"). */
+	static TMap<FName, FString> TypeToRole;
+
+	/** Prefix-based role mapping for template types (e.g., "SComboBox" matches "SComboBox<...>"). */
+	static TArray<TPair<FString, FString>> TemplatePrefixToRole;
+
+	/** Map from role string to ref prefix (e.g., "button" -> "b"). */
+	static TMap<FString, FString> RoleToRefPrefix;
+
+	/** Set of structural container type names that should be skipped (children promoted). */
+	static TSet<FName> StructuralContainers;
+
+	/** Custom label extractors registered via RegisterLabelExtractor. */
+	static TMap<FName, TFunction<FString(TSharedRef<SWidget>)>> CustomLabelExtractors;
+
+	/** Custom state flags extractors registered via RegisterStateFlagsExtractor. */
+	static TMap<FName, TFunction<TArray<FString>(TSharedRef<SWidget>)>> CustomStateFlagsExtractors;
+
+	static bool bMapsInitialized;
+};
diff --git a/Engine/Plugins/Experimental/SlateInspector/Source/SlateInspector/Public/SlateInspectorSubsystem.h b/Engine/Plugins/Experimental/SlateInspector/Source/SlateInspector/Public/SlateInspectorSubsystem.h
new file mode 100644
--- /dev/null
+++ b/Engine/Plugins/Experimental/SlateInspector/Source/SlateInspector/Public/SlateInspectorSubsystem.h
@@ -0,0 +1,28 @@
+// Copyright Epic Games, Inc. All Rights Reserved.
+
+#pragma once
+
+#include "EditorSubsystem.h"
+
+#include "SlateInspectorSubsystem.generated.h"
+
+class FSubsystemCollectionBase;
+
+UCLASS(MinimalAPI)
+class USlateInspectorSubsystem : public UEditorSubsystem
+{
+	GENERATED_BODY()
+
+public:
+
+	virtual void Initialize(FSubsystemCollectionBase& Collection) override;
+
+	virtual void Deinitialize() override;
+
+	/** Register or unregister the toolset at runtime. */
+	void SetToolsetEnabled(bool bEnabled);
+
+private:
+
+	bool bToolsetRegistered = false;
+};
diff --git a/Engine/Plugins/Experimental/SlateInspector/Source/SlateInspector/Public/SlateInspectorToolset.h b/Engine/Plugins/Experimental/SlateInspector/Source/SlateInspector/Public/SlateInspectorToolset.h
new file mode 100644
--- /dev/null
+++ b/Engine/Plugins/Experimental/SlateInspector/Source/SlateInspector/Public/SlateInspectorToolset.h
@@ -0,0 +1,141 @@
+// Copyright Epic Games, Inc. All Rights Reserved.
+
+#pragma once
+
+#include "ToolsetRegistry/UToolsetRegistry.h"
+#include "ToolsetRegistry/ToolsetImage.h"
+#include "Widgets/SWidget.h"
+
+#include "SlateInspectorToolset.generated.h"
+
+/** Describes a single form field for the FillForm tool. */
+USTRUCT(BlueprintType)
+struct FSlateInspectorFormField
+{
+	GENERATED_BODY()
+
+	/** The ref identifier of the form field widget. */
+	UPROPERTY(BlueprintReadWrite, Category = "SlateInspector")
+	FString Ref;
+
+	/** The value to set. */
+	UPROPERTY(BlueprintReadWrite, Category = "SlateInspector")
+	FString Value;
+
+	/** The type of field: "textbox", "checkbox", "combobox", or "slider". */
+	UPROPERTY(BlueprintReadWrite, Category = "SlateInspector")
+	FString FieldType;
+};
+
+/**
+ * Playwright-style Slate UI automation toolset.
+ *
+ * Exposes snapshot, screenshot, and interaction tools for driving the
+ * Unreal Editor UI programmatically.  Registered via UToolsetRegistry
+ * so the ModelContextProtocol plugin picks them up automatically.
+ *
+ * Input simulation uses direct Slate event APIs (ProcessKeyCharEvent,
+ * ProcessMouseButtonDownEvent, etc.) rather than the AutomationDriver,
+ * because AutomationDriver's synchronous API seems to deadlock when called
+ * from the game thread (which is where MCP tool calls execute).
+ */
+UCLASS(BlueprintType, MinimalAPI)
+class USlateInspectorToolset : public UToolsetDefinition
+{
+	GENERATED_BODY()
+
+public:
+
+	/** Capture a Slate UI accessibility snapshot. Use this BEFORE any other SlateInspector
+	 * tool -- it returns the refs needed by Click, Type, Hover, Screenshot, etc.
+	 * IMPORTANT: Every Snapshot call resets ALL refs. Never reuse refs from a prior Snapshot.
+	 * @param Ref                      Subtree root ref (from most recent Snapshot). Empty = all windows.
+	 * @param MaxDepth                 Maximum depth (default 30).
+	 * @param bIncludeSourceLocations  Include [src=File:Line] tags showing where each widget was created in C++. */
+	UFUNCTION(BlueprintCallable, Category = "SlateInspector")
+	static FString Snapshot(const FString& Ref, int32 MaxDepth = 30, bool bIncludeSourceLocations = false);
+
+	/** Screenshot a Slate widget or the active editor window. Prefer this over
+	 * SceneTools.take_screenshot for Editor UI; use SceneTools only for 3D viewport.
+	 * @param Ref  Slate widget ref (from most recent Snapshot). Empty = active window. */
+	UFUNCTION(BlueprintCallable, Category = "SlateInspector")
+	static FToolsetImage Screenshot(const FString& Ref);
+
+	// -- Action tools (all refs must come from the most recent Snapshot) --
+
+	/** Click a Slate widget identified by its snapshot ref.
+	 * Refs reset each Snapshot -- never reuse old refs.
+	 * @param Ref         Slate widget ref from the most recent Snapshot.
+	 * @param Button      "left", "right", or "middle".
+	 * @param DoubleClick True for double-click. */
+	UFUNCTION(BlueprintCallable, Category = "SlateInspector")
+	static bool Click(const FString& Ref, const FString& Button = TEXT("left"), bool DoubleClick = false);
+
+	/** Hover over a Slate widget, triggering any hover state or tooltip.
+	 * Refs reset each Snapshot -- never reuse old refs.
+	 * @param Ref  Slate widget ref from the most recent Snapshot. */
+	UFUNCTION(BlueprintCallable, Category = "SlateInspector")
+	static bool Hover(const FString& Ref);
+
+	/** Type text into a Slate text input widget. Focuses the widget first,
+	 * then sends one key event per character.
+	 * Refs reset each Snapshot -- never reuse old refs.
+	 * @param Ref    Slate textbox ref from the most recent Snapshot.
+	 * @param Text   Text to type.
+	 * @param Submit Press Enter after typing. */
+	UFUNCTION(BlueprintCallable, Category = "SlateInspector")
+	static bool Type(const FString& Ref, const FString& Text, bool Submit = false);
+
+	/** Press and release a keyboard key on the currently focused Slate widget.
+	 * @param Key  Key name, e.g. "Enter", "Escape", "Tab", "ArrowLeft". */
+	UFUNCTION(BlueprintCallable, Category = "SlateInspector")
+	static bool PressKey(const FString& Key);
+
+	/** Select an option in a Slate combobox by its text label.
+	 * Opens the dropdown, finds the matching text, and clicks it.
+	 * Refs reset each Snapshot -- never reuse old refs.
+	 * @param Ref    Slate combobox ref from the most recent Snapshot.
+	 * @param Value  Exact option text to select. */
+	UFUNCTION(BlueprintCallable, Category = "SlateInspector")
+	static bool SelectOption(const FString& Ref, const FString& Value);
+
+	/** Drag from one Slate widget to another (mouse down, move, release).
+	 * Refs reset each Snapshot -- never reuse old refs.
+	 * @param StartRef  Slate widget ref for the drag source.
+	 * @param EndRef    Slate widget ref for the drop target. */
+	UFUNCTION(BlueprintCallable, Category = "SlateInspector")
+	static bool Drag(const FString& StartRef, const FString& EndRef);
+
+	/** List, select, or close top-level Slate editor windows.
+	 * @param Action  "list" returns JSON array, "select" brings to front, "close" destroys.
+	 * @param Index   Window index for select/close. */
+	UFUNCTION(BlueprintCallable, Category = "SlateInspector")
+	static FString Windows(const FString& Action = TEXT("list"), int32 Index = -1);
+
+	/** Check if text is present or absent in the Slate widget tree.
+	 * Non-blocking: checks once and returns immediately. Poll to wait.
+	 * @param Text       Text that must be present (empty = skip).
+	 * @param TextGone   Text that must be absent (empty = skip). */
+	UFUNCTION(BlueprintCallable, Category = "SlateInspector")
+	static bool WaitFor(const FString& Text = TEXT(""), const FString& TextGone = TEXT(""));
+
+	/** Fill multiple Slate form fields at once. Refs from most recent Snapshot only.
+	 * @param Fields  Array of {Ref, Value, FieldType} where FieldType is
+	 *                "textbox", "checkbox", "combobox", or "slider". */
+	UFUNCTION(BlueprintCallable, Category = "SlateInspector")
+	static bool FillForm(const TArray<FSlateInspectorFormField>& Fields);
+
+private:
+
+	/** Resolve a ref to a widget, logging a warning if it fails. */
+	static TSharedPtr<SWidget> ResolveRefOrWarn(const FString& Ref);
+
+	/** Get the absolute screen-space center of a widget. */
+	static FVector2D GetWidgetScreenCenter(TSharedRef<SWidget> Widget);
+
+	/** Map a button name ("left", "right", "middle") to an FKey. */
+	static FKey ParseMouseButtonKey(const FString& Button);
+
+	/** Simulate a mouse click at a widget's screen-space center. */
+	static bool SimulateClick(TSharedRef<SWidget> Widget, FKey MouseButton, bool bDoubleClick);
+};
diff --git a/Engine/Plugins/Experimental/SlateInspector/Source/SlateInspector/SlateInspector.Build.cs b/Engine/Plugins/Experimental/SlateInspector/Source/SlateInspector/SlateInspector.Build.cs
new file mode 100644
--- /dev/null
+++ b/Engine/Plugins/Experimental/SlateInspector/Source/SlateInspector/SlateInspector.Build.cs
@@ -0,0 +1,30 @@
+// Copyright Epic Games, Inc. All Rights Reserved.
+
+namespace UnrealBuildTool.Rules
+{
+	public class SlateInspector : ModuleRules
+	{
+		public SlateInspector(ReadOnlyTargetRules Target) : base(Target)
+		{
+			PublicDependencyModuleNames.AddRange(
+				new string[] {
+					"ToolsetRegistry",
+				}
+			);
+
+			PrivateDependencyModuleNames.AddRange(
+				new string[] {
+					"Core",
+					"CoreUObject",
+					"EditorSubsystem",
+					"Engine",
+					"ImageWrapper",
+					"InputCore",
+					"Slate",
+					"SlateCore",
+					"UnrealEd",
+				}
+			);
+		}
+	}
+}
