[SlateInspector] Fix `SimulateClick` not firing `SButton::OnClicked`.

- Add a `ProcessMouseMoveEvent` call before mouse down/up in `SimulateClick`. `SButton` (`DownAndUp` click method) requires `IsHovered() == true` on mouse-up to fire `OnClicked`. Previously, `SetCursorPos` moved the OS cursor but did not update Slate hover tracking, so the button never registered as hovered and `OnClicked` was silently skipped.

#Synced-CL 51849904

diff --git a/Engine/Plugins/Experimental/SlateInspector/Source/SlateInspector/Private/SlateInspectorToolset.cpp b/Engine/Plugins/Experimental/SlateInspector/Source/SlateInspector/Private/SlateInspectorToolset.cpp
--- a/Engine/Plugins/Experimental/SlateInspector/Source/SlateInspector/Private/SlateInspectorToolset.cpp
+++ b/Engine/Plugins/Experimental/SlateInspector/Source/SlateInspector/Private/SlateInspectorToolset.cpp
@@ -100,6 +100,21 @@ bool USlateInspectorToolset::SimulateClick(TSharedRef<SWidget> Widget, FKey Mous
 	FSlateApplication::Get().SetCursorPos(ScreenCenter);
 
 	FModifierKeysState ModifierKeys;
+
+	// Send a mouse move event first so Slate updates hover tracking.
+	// SButton (DownAndUp click method) requires IsHovered() == true on
+	// mouse-up to fire OnClicked. Without this, SetCursorPos alone
+	// moves the OS cursor but does not update Slate's hover state.
+	FPointerEvent MoveEvent(
+		0,
+		ScreenCenter,
+		ScreenCenter,
+		TSet<FKey>(),
+		EKeys::Invalid,
+		0.0f,
+		ModifierKeys);
+	FSlateApplication::Get().ProcessMouseMoveEvent(MoveEvent);
+
 	TSet<FKey> PressedButtons;
 	PressedButtons.Add(MouseButton);
 
