[SlateInspector] Add *observer-based* ref system, test module, and input simulation fixes.

Replace the reset-per-snapshot ref model with continuous observers that walk watched widget subtrees on a throttled tick (~100ms), keeping the ref cache up to date as the UI changes.

# Observer-Based System
- New tools: `Observe`, `Unobserve`, `ListObservers`
- New class: `FSlateInspectorObserverManager` with throttled tick, deduplication, and dead-root cleanup
- New methods: `FSlateInspectorRefCache::PurgeExpired`, `FSlateInspectorRefCache::FindRef`
- `Render` gains `bResetCache` param; `Snapshot` returns cached observer text when available
- `USlateInspectorSubsystem` hooks `FSlateApplication::OnPostTick` and creates a depth-0 root observer

# Input Simulation
- New struct `FSlateInspectorModifierKeys` for passing modifier key state (*Shift*, *Ctrl*, *Alt*, *Cmd*) to `Click` and `Drag`
- `SimulateClick` walks parent chain to find the correct `SWindow` (fixes misrouted clicks on `SCheckBox`, etc.)
- `SimulateClick` sends mouse-move before click and mouse-up for both single/double-click paths
- `PressKey` supports modifier prefixes (*Ctrl+A*, *Shift+Tab*) and sends `ProcessKeyCharEvent` for *BackSpace*/*Tab*/*Enter*/*Escape*
- `Hover` sends `ProcessMouseMoveEvent` so hover state actually updates
- `SelectOption` prioritizes newly appeared popup windows and double-ticks to materialize `SListView` items
- `Type` verifies keyboard focus; `FillForm` warns on type mismatch
- `Drag` walks parent chain to find `SWindow`, sends mouse-move before pressing

# Test Module
- 47 CQTest tests covering all tools against `SSlateInspectorTestPanel` (button, textbox, checkbox, combobox, slider, spinbox, drag, etc.)

#Synced-CL 51849904

diff --git a/Engine/Plugins/Experimental/SlateInspector/SlateInspector.uplugin b/Engine/Plugins/Experimental/SlateInspector/SlateInspector.uplugin
--- a/Engine/Plugins/Experimental/SlateInspector/SlateInspector.uplugin
+++ b/Engine/Plugins/Experimental/SlateInspector/SlateInspector.uplugin
@@ -20,6 +20,11 @@
 			"Name": "SlateInspector",
 			"Type": "Editor",
 			"LoadingPhase": "Default"
+		},
+		{
+			"Name": "SlateInspectorTests",
+			"Type": "Editor",
+			"LoadingPhase": "Default"
 		}
 	],
 	"Plugins": [
diff --git a/Engine/Plugins/Experimental/SlateInspector/Source/SlateInspector/Private/SlateInspectorObserverManager.cpp b/Engine/Plugins/Experimental/SlateInspector/Source/SlateInspector/Private/SlateInspectorObserverManager.cpp
new file mode 100644
--- /dev/null
+++ b/Engine/Plugins/Experimental/SlateInspector/Source/SlateInspector/Private/SlateInspectorObserverManager.cpp
@@ -0,0 +1,152 @@
+// Copyright Epic Games, Inc. All Rights Reserved.
+
+#include "SlateInspectorObserverManager.h"
+
+#include "SlateInspectorRefCache.h"
+#include "SlateInspectorSnapshotRenderer.h"
+
+FSlateInspectorObserverManager& FSlateInspectorObserverManager::Get()
+{
+	static FSlateInspectorObserverManager Instance;
+	return Instance;
+}
+
+FString FSlateInspectorObserverManager::AddObserver(TSharedPtr<SWidget> RootWidget, int32 MaxDepth)
+{
+	MaxDepth = FMath::Max(MaxDepth, 0);
+	bool bRoot = !RootWidget.IsValid();
+
+	// If an observer already exists for the same root, update its depth and return its identifier.
+	for (FSlateInspectorObserver& Existing : Observers)
+	{
+		if (bRoot && Existing.bRoot)
+		{
+			Existing.MaxDepth = FMath::Max(Existing.MaxDepth, MaxDepth);
+			return Existing.Identifier;
+		}
+		if (!bRoot && !Existing.bRoot && Existing.RootWidget.Pin() == RootWidget)
+		{
+			Existing.MaxDepth = FMath::Max(Existing.MaxDepth, MaxDepth);
+			return Existing.Identifier;
+		}
+	}
+
+	FSlateInspectorObserver Observer;
+	Observer.Identifier = FString::Printf(TEXT("observer_%d"), NextIdentifier++);
+	Observer.bRoot = bRoot;
+	Observer.RootWidget = RootWidget;
+	Observer.MaxDepth = MaxDepth;
+
+	Observers.Add(MoveTemp(Observer));
+
+	return Observers.Last().Identifier;
+}
+
+bool FSlateInspectorObserverManager::RemoveObserver(const FString& Identifier)
+{
+	int32 Index = Observers.IndexOfByPredicate([&Identifier](const FSlateInspectorObserver& Observer)
+	{
+		return Observer.Identifier == Identifier;
+	});
+
+	if (Index == INDEX_NONE)
+	{
+		return false;
+	}
+
+	Observers.RemoveAt(Index);
+	return true;
+}
+
+TArray<FSlateInspectorObserver> FSlateInspectorObserverManager::GetObservers() const
+{
+	return Observers;
+}
+
+FString FSlateInspectorObserverManager::GetCachedSnapshot(const FString& Identifier) const
+{
+	for (const FSlateInspectorObserver& Observer : Observers)
+	{
+		if (Observer.Identifier == Identifier)
+		{
+			return Observer.CachedSnapshotText;
+		}
+	}
+
+	return FString();
+}
+
+FString FSlateInspectorObserverManager::FindMatchingObserverSnapshot(TSharedPtr<SWidget> RootWidget, int32 MaxDepth) const
+{
+	const FSlateInspectorObserver* BestMatch = nullptr;
+
+	for (const FSlateInspectorObserver& Observer : Observers)
+	{
+		if (Observer.MaxDepth < MaxDepth)
+		{
+			continue;
+		}
+
+		if (!RootWidget.IsValid() && Observer.bRoot)
+		{
+			if (!BestMatch || Observer.MaxDepth > BestMatch->MaxDepth)
+			{
+				BestMatch = &Observer;
+			}
+		}
+		else if (RootWidget.IsValid() && !Observer.bRoot)
+		{
+			TSharedPtr<SWidget> ObserverRoot = Observer.RootWidget.Pin();
+			if (ObserverRoot == RootWidget)
+			{
+				if (!BestMatch || Observer.MaxDepth > BestMatch->MaxDepth)
+				{
+					BestMatch = &Observer;
+				}
+			}
+		}
+	}
+
+	return BestMatch ? BestMatch->CachedSnapshotText : FString();
+}
+
+void FSlateInspectorObserverManager::Tick(float DeltaTime)
+{
+	AccumulatedTime += DeltaTime;
+
+	if (AccumulatedTime < TickIntervalSeconds)
+	{
+		return;
+	}
+	AccumulatedTime = 0.0f;
+
+	if (Observers.IsEmpty())
+	{
+		return;
+	}
+
+	// Remove non-root observers whose root widget has been destroyed.
+	Observers.RemoveAll([](const FSlateInspectorObserver& Observer)
+	{
+		return !Observer.bRoot && !Observer.RootWidget.IsValid();
+	});
+
+	// Walk each observer's subtree and cache the snapshot text.
+	for (FSlateInspectorObserver& Observer : Observers)
+	{
+		TSharedPtr<SWidget> Root = Observer.bRoot ? nullptr : Observer.RootWidget.Pin();
+		Observer.CachedSnapshotText = FSlateInspectorSnapshotRenderer::Render(
+			Root,
+			Observer.MaxDepth,
+			/* bIncludeSourceLocations */ false,
+			/* bResetCache */ false);
+	}
+
+	// Periodically purge expired refs.
+	++TickCounter;
+	if (TickCounter >= PurgeEveryNTicks)
+	{
+		TickCounter = 0;
+		FSlateInspectorRefCache::Get().PurgeExpired();
+	}
+}
diff --git a/Engine/Plugins/Experimental/SlateInspector/Source/SlateInspector/Private/SlateInspectorRefCache.cpp b/Engine/Plugins/Experimental/SlateInspector/Source/SlateInspector/Private/SlateInspectorRefCache.cpp
--- a/Engine/Plugins/Experimental/SlateInspector/Source/SlateInspector/Private/SlateInspectorRefCache.cpp
+++ b/Engine/Plugins/Experimental/SlateInspector/Source/SlateInspector/Private/SlateInspectorRefCache.cpp
@@ -28,6 +28,16 @@ FString FSlateInspectorRefCache::GetOrAssignRef(TSharedRef<SWidget> Widget, cons
 	return Ref;
 }
 
+FString FSlateInspectorRefCache::FindRef(TSharedRef<SWidget> Widget) const
+{
+	TWeakPtr<SWidget> WeakWidget = Widget;
+	if (const FString* Ref = WidgetToRef.Find(WeakWidget))
+	{
+		return *Ref;
+	}
+	return FString();
+}
+
 TSharedPtr<SWidget> FSlateInspectorRefCache::ResolveRef(const FString& Ref) const
 {
 	if (const TWeakPtr<SWidget>* WeakWidget = RefToWidget.Find(Ref))
@@ -43,3 +53,24 @@ void FSlateInspectorRefCache::Reset()
 	RefToWidget.Empty();
 	RolePrefixCounters.Empty();
 }
+
+void FSlateInspectorRefCache::PurgeExpired()
+{
+	// Capture both the key and its ref string in a single pass.  Expired
+	// weak pointers may hash identically, so a second Find() on the key
+	// alone could silently fail or match the wrong entry.
+	TArray<TPair<TWeakPtr<SWidget>, FString>> Expired;
+	for (const auto& Pair : WidgetToRef)
+	{
+		if (!Pair.Key.IsValid())
+		{
+			Expired.Emplace(Pair.Key, Pair.Value);
+		}
+	}
+
+	for (const auto& Entry : Expired)
+	{
+		RefToWidget.Remove(Entry.Value);
+		WidgetToRef.Remove(Entry.Key);
+	}
+}
diff --git a/Engine/Plugins/Experimental/SlateInspector/Source/SlateInspector/Private/SlateInspectorSnapshotRenderer.cpp b/Engine/Plugins/Experimental/SlateInspector/Source/SlateInspector/Private/SlateInspectorSnapshotRenderer.cpp
--- a/Engine/Plugins/Experimental/SlateInspector/Source/SlateInspector/Private/SlateInspectorSnapshotRenderer.cpp
+++ b/Engine/Plugins/Experimental/SlateInspector/Source/SlateInspector/Private/SlateInspectorSnapshotRenderer.cpp
@@ -105,11 +105,13 @@ void FSlateInspectorSnapshotRenderer::RegisterWidgetRole(const FName& WidgetType
 
 void FSlateInspectorSnapshotRenderer::RegisterLabelExtractor(const FName& WidgetType, TFunction<FString(TSharedRef<SWidget>)> Extractor)
 {
+	EnsureMapsInitialized();
 	CustomLabelExtractors.Add(WidgetType, MoveTemp(Extractor));
 }
 
 void FSlateInspectorSnapshotRenderer::RegisterStateFlagsExtractor(const FName& WidgetType, TFunction<TArray<FString>(TSharedRef<SWidget>)> Extractor)
 {
+	EnsureMapsInitialized();
 	CustomStateFlagsExtractors.Add(WidgetType, MoveTemp(Extractor));
 }
 
@@ -133,11 +135,16 @@ const FString* FSlateInspectorSnapshotRenderer::FindRoleForType(const FName& Wid
 	return nullptr;
 }
 
-FString FSlateInspectorSnapshotRenderer::Render(TSharedPtr<SWidget> Root, int32 MaxDepth, bool bIncludeSourceLocations)
+FString FSlateInspectorSnapshotRenderer::Render(TSharedPtr<SWidget> Root, int32 MaxDepth, bool bIncludeSourceLocations, bool bResetCache)
 {
 	EnsureMapsInitialized();
 
-	FSlateInspectorRefCache::Get().Reset();
+	MaxDepth = FMath::Max(MaxDepth, 0);
+
+	if (bResetCache)
+	{
+		FSlateInspectorRefCache::Get().Reset();
+	}
 
 	TStringBuilder<4096> Output;
 
@@ -206,7 +213,8 @@ void FSlateInspectorSnapshotRenderer::RenderWidget(
 	if (Children->Num() == 1)
 	{
 		TSharedRef<SWidget> OnlyChild = Children->GetChildAt(0);
-		if (OnlyChild->GetType() == FName(TEXT("STextBlock")))
+		static const FName TextBlockType(TEXT("STextBlock"));
+		if (OnlyChild->GetType() == TextBlockType)
 		{
 			TSharedRef<STextBlock> TextBlock = StaticCastSharedRef<STextBlock>(OnlyChild);
 			AbsorbedText = TextBlock->GetText().ToString();
@@ -228,7 +236,7 @@ void FSlateInspectorSnapshotRenderer::RenderWidget(
 		}
 	}
 
-	// This widget is meaningful — emit it.
+	// This widget is meaningful, emit it.
 	FString Label;
 	if (bAbsorbedTextChild && !AbsorbedText.IsEmpty())
 	{
@@ -324,7 +332,7 @@ void FSlateInspectorSnapshotRenderer::RenderWidget(
 
 FString FSlateInspectorSnapshotRenderer::ExtractLabel(TSharedRef<SWidget> Widget, const FString& TypeString)
 {
-	// Prefer Slate's native accessible text — many widgets implement this already.
+	// Prefer Slate's native accessible text since many widgets implement this already.
 #if WITH_ACCESSIBILITY
 	FText AccessibleText = Widget->GetAccessibleText(EAccessibleType::Main);
 	if (!AccessibleText.IsEmpty())
diff --git a/Engine/Plugins/Experimental/SlateInspector/Source/SlateInspector/Private/SlateInspectorSubsystem.cpp b/Engine/Plugins/Experimental/SlateInspector/Source/SlateInspector/Private/SlateInspectorSubsystem.cpp
--- a/Engine/Plugins/Experimental/SlateInspector/Source/SlateInspector/Private/SlateInspectorSubsystem.cpp
+++ b/Engine/Plugins/Experimental/SlateInspector/Source/SlateInspector/Private/SlateInspectorSubsystem.cpp
@@ -3,9 +3,11 @@
 #include "SlateInspectorSubsystem.h"
 
 #include "Editor.h"
-#include "ToolsetRegistry/UToolsetRegistry.h"
+#include "Framework/Application/SlateApplication.h"
+#include "SlateInspectorObserverManager.h"
 #include "SlateInspectorToolset.h"
 #include "Subsystems/SubsystemCollection.h"
+#include "ToolsetRegistry/UToolsetRegistry.h"
 
 static bool bEnableSlateInspector = true;
 
@@ -33,14 +35,39 @@ void USlateInspectorSubsystem::Initialize(FSubsystemCollectionBase& Collection)
 {
 	Super::Initialize(Collection);
 	SetToolsetEnabled(bEnableSlateInspector);
+
+	if (FSlateApplication::IsInitialized())
+	{
+		PostTickDelegateHandle = FSlateApplication::Get().OnPostTick().AddUObject(this, &USlateInspectorSubsystem::OnSlatePostTick);
+	}
+
+	// Create a shallow default root observer (depth 0: window names only).
+	// Use Observe() on specific windows or panels to drill deeper.
+	RootObserverIdentifier = FSlateInspectorObserverManager::Get().AddObserver(nullptr, 0);
 }
 
 void USlateInspectorSubsystem::Deinitialize()
 {
+	if (!RootObserverIdentifier.IsEmpty())
+	{
+		FSlateInspectorObserverManager::Get().RemoveObserver(RootObserverIdentifier);
+		RootObserverIdentifier.Empty();
+	}
+
+	if (FSlateApplication::IsInitialized())
+	{
+		FSlateApplication::Get().OnPostTick().Remove(PostTickDelegateHandle);
+	}
+
 	SetToolsetEnabled(false);
 	Super::Deinitialize();
 }
 
+void USlateInspectorSubsystem::OnSlatePostTick(float DeltaTime)
+{
+	FSlateInspectorObserverManager::Get().Tick(DeltaTime);
+}
+
 void USlateInspectorSubsystem::SetToolsetEnabled(bool bEnabled)
 {
 	if (bEnabled && !bToolsetRegistered)
diff --git a/Engine/Plugins/Experimental/SlateInspector/Source/SlateInspector/Private/SlateInspectorToolset.cpp b/Engine/Plugins/Experimental/SlateInspector/Source/SlateInspector/Private/SlateInspectorToolset.cpp
--- a/Engine/Plugins/Experimental/SlateInspector/Source/SlateInspector/Private/SlateInspectorToolset.cpp
+++ b/Engine/Plugins/Experimental/SlateInspector/Source/SlateInspector/Private/SlateInspectorToolset.cpp
@@ -2,6 +2,7 @@
 
 #include "SlateInspectorToolset.h"
 
+#include "SlateInspectorObserverManager.h"
 #include "SlateInspectorRefCache.h"
 #include "SlateInspectorSnapshotRenderer.h"
 
@@ -82,29 +83,40 @@ FKey USlateInspectorToolset::ParseMouseButtonKey(const FString& Button)
 	return EKeys::LeftMouseButton;
 }
 
-bool USlateInspectorToolset::SimulateClick(TSharedRef<SWidget> Widget, FKey MouseButton, bool bDoubleClick)
+bool USlateInspectorToolset::SimulateClick(TSharedRef<SWidget> Widget, FKey MouseButton, bool bDoubleClick, const FModifierKeysState& ModifierKeys)
 {
 	if (!FSlateApplication::IsInitialized())
 	{
 		return false;
 	}
 
-	TSharedPtr<SWindow> Window = FSlateApplication::Get().FindWidgetWindow(Widget);
+	// Find the window containing this widget by walking up its parent chain
+	// to the nearest SWindow. FindWidgetWindow can return the wrong window
+	// for some widget types, and LocateWindowUnderMouse depends on z-order.
+	TSharedPtr<SWindow> Window;
+	TSharedPtr<SWidget> Current = Widget;
+	static const FName WindowType(TEXT("SWindow"));
+	while (Current.IsValid())
+	{
+		if (Current->GetType() == WindowType)
+		{
+			Window = StaticCastSharedPtr<SWindow>(Current);
+			break;
+		}
+		Current = Current->GetParentWidget();
+	}
 	if (!Window.IsValid())
 	{
-		UE_LOGFMT(LogSlateInspector, Warning, "SlateInspector: Could not find window for widget");
+		UE_LOGFMT(LogSlateInspector, Warning, "SlateInspector: Could not find window containing widget");
 		return false;
 	}
 
 	FVector2D ScreenCenter = GetWidgetScreenCenter(Widget);
 	FSlateApplication::Get().SetCursorPos(ScreenCenter);
 
-	FModifierKeysState ModifierKeys;
-
-	// Send a mouse move event first so Slate updates hover tracking.
-	// SButton (DownAndUp click method) requires IsHovered() == true on
-	// mouse-up to fire OnClicked. Without this, SetCursorPos alone
-	// moves the OS cursor but does not update Slate's hover state.
+	// Send a mouse move so Slate updates its hover tracking, widget path
+	// cache, and window z-order for hit-testing. This must happen after
+	// BringToFront so the window reorder is reflected.
 	FPointerEvent MoveEvent(
 		0,
 		ScreenCenter,
@@ -134,19 +146,20 @@ bool USlateInspectorToolset::SimulateClick(TSharedRef<SWidget> Widget, FKey Mous
 	else
 	{
 		FSlateApplication::Get().ProcessMouseButtonDownEvent(Window->GetNativeWindow(), MouseEvent);
-
-		// Release with empty pressed buttons set.
-		FPointerEvent MouseUpEvent(
-			0,
-			ScreenCenter,
-			ScreenCenter,
-			TSet<FKey>(),
-			MouseButton,
-			0.0f,
-			ModifierKeys);
-		FSlateApplication::Get().ProcessMouseButtonUpEvent(MouseUpEvent);
 	}
 
+	// Release for both paths. ProcessMouseButtonDoubleClickEvent implies a
+	// mouse-down, so we must release to avoid leaving widgets in pressed state.
+	FPointerEvent MouseUpEvent(
+		0,
+		ScreenCenter,
+		ScreenCenter,
+		TSet<FKey>(),
+		MouseButton,
+		0.0f,
+		ModifierKeys);
+	FSlateApplication::Get().ProcessMouseButtonUpEvent(MouseUpEvent);
+
 	return true;
 }
 
@@ -154,6 +167,8 @@ bool USlateInspectorToolset::SimulateClick(TSharedRef<SWidget> Widget, FKey Mous
 
 FString USlateInspectorToolset::Snapshot(const FString& Ref, int32 MaxDepth, bool bIncludeSourceLocations)
 {
+	MaxDepth = FMath::Max(MaxDepth, 0);
+
 	TSharedPtr<SWidget> Root;
 	if (!Ref.IsEmpty())
 	{
@@ -164,7 +179,77 @@ FString USlateInspectorToolset::Snapshot(const FString& Ref, int32 MaxDepth, boo
 		}
 	}
 
-	return FSlateInspectorSnapshotRenderer::Render(Root, MaxDepth, bIncludeSourceLocations);
+	// If an observer covers this root, return its cached snapshot.
+	if (!bIncludeSourceLocations)
+	{
+		FString CachedSnapshot = FSlateInspectorObserverManager::Get().FindMatchingObserverSnapshot(Root, MaxDepth);
+		if (!CachedSnapshot.IsEmpty())
+		{
+			return CachedSnapshot;
+		}
+	}
+
+	// No observer covers this request. Render fresh without resetting the cache
+	// (observers may be maintaining refs for other subtrees).
+	return FSlateInspectorSnapshotRenderer::Render(Root, MaxDepth, bIncludeSourceLocations, /* bResetCache */ false);
+}
+
+FString USlateInspectorToolset::Observe(const FString& Ref, int32 MaxDepth)
+{
+	MaxDepth = FMath::Max(MaxDepth, 0);
+
+	TSharedPtr<SWidget> RootWidget;
+	if (!Ref.IsEmpty())
+	{
+		RootWidget = ResolveRefOrWarn(Ref);
+		if (!RootWidget.IsValid())
+		{
+			return TEXT("Error: Could not resolve ref");
+		}
+	}
+
+	return FSlateInspectorObserverManager::Get().AddObserver(RootWidget, MaxDepth);
+}
+
+bool USlateInspectorToolset::Unobserve(const FString& Identifier)
+{
+	return FSlateInspectorObserverManager::Get().RemoveObserver(Identifier);
+}
+
+FString USlateInspectorToolset::ListObservers()
+{
+	TArray<FSlateInspectorObserver> Observers = FSlateInspectorObserverManager::Get().GetObservers();
+
+	FString Result = TEXT("[");
+	for (int32 Index = 0; Index < Observers.Num(); ++Index)
+	{
+		const FSlateInspectorObserver& Observer = Observers[Index];
+		if (Index > 0)
+		{
+			Result += TEXT(", ");
+		}
+
+		FString RootRef;
+		if (!Observer.bRoot)
+		{
+			TSharedPtr<SWidget> RootWidget = Observer.RootWidget.Pin();
+			if (RootWidget.IsValid())
+			{
+				RootRef = FSlateInspectorRefCache::Get().FindRef(RootWidget.ToSharedRef());
+			}
+		}
+
+		Result += FString::Printf(
+			TEXT("{\"identifier\": \"%s\", \"root\": %s, \"rootRef\": \"%s\", \"maxDepth\": %d, \"cachedSnapshotSize\": %d}"),
+			*SlateInspectorPrivate::EscapeJsonString(Observer.Identifier),
+			Observer.bRoot ? TEXT("true") : TEXT("false"),
+			*SlateInspectorPrivate::EscapeJsonString(RootRef),
+			Observer.MaxDepth,
+			Observer.CachedSnapshotText.Len());
+	}
+	Result += TEXT("]");
+
+	return Result;
 }
 
 FToolsetImage USlateInspectorToolset::Screenshot(const FString& Ref)
@@ -205,15 +290,20 @@ FToolsetImage USlateInspectorToolset::Screenshot(const FString& Ref)
 
 // ── Action tools ─────────────────────────────────────────────────────
 
-bool USlateInspectorToolset::Click(const FString& Ref, const FString& Button, bool DoubleClick)
+bool USlateInspectorToolset::Click(const FString& Ref, const FString& Button, bool DoubleClick, const FSlateInspectorModifierKeys& Modifiers)
 {
+	if (!FSlateApplication::IsInitialized())
+	{
+		return false;
+	}
+
 	TSharedPtr<SWidget> Widget = ResolveRefOrWarn(Ref);
 	if (!Widget.IsValid())
 	{
 		return false;
 	}
 
-	return SimulateClick(Widget.ToSharedRef(), ParseMouseButtonKey(Button), DoubleClick);
+	return SimulateClick(Widget.ToSharedRef(), ParseMouseButtonKey(Button), DoubleClick, Modifiers.ToModifierKeysState());
 }
 
 bool USlateInspectorToolset::Hover(const FString& Ref)
@@ -226,13 +316,27 @@ bool USlateInspectorToolset::Hover(const FString& Ref)
 
 	FVector2D ScreenCenter = GetWidgetScreenCenter(Widget.ToSharedRef());
 	FSlateApplication::Get().SetCursorPos(ScreenCenter);
+
+	// SetCursorPos alone does not update Slate's hover tracking.
+	// Send a mouse move event so the widget receives OnMouseEnter and IsHovered() becomes true.
+	FModifierKeysState ModifierKeys;
+	FPointerEvent MouseMoveEvent(
+		0, ScreenCenter, ScreenCenter,
+		TSet<FKey>(), FKey(), 0.0f, ModifierKeys);
+	FSlateApplication::Get().ProcessMouseMoveEvent(MouseMoveEvent);
+
 	return true;
 }
 
 bool USlateInspectorToolset::Type(const FString& Ref, const FString& Text, bool Submit)
 {
+	if (!FSlateApplication::IsInitialized())
+	{
+		return false;
+	}
+
 	TSharedPtr<SWidget> Widget = ResolveRefOrWarn(Ref);
-	if (!Widget.IsValid() || !FSlateApplication::IsInitialized())
+	if (!Widget.IsValid())
 	{
 		return false;
 	}
@@ -240,6 +344,12 @@ bool USlateInspectorToolset::Type(const FString& Ref, const FString& Text, bool
 	// Focus the target widget.
 	FSlateApplication::Get().SetKeyboardFocus(Widget, EFocusCause::SetDirectly);
 
+	if (!Widget->HasKeyboardFocus())
+	{
+		UE_LOGFMT(LogSlateInspector, Warning, "SlateInspector: Could not focus widget for typing");
+		return false;
+	}
+
 	FModifierKeysState ModifierKeys;
 
 	// Send each character as a key char event.
@@ -267,8 +377,42 @@ bool USlateInspectorToolset::PressKey(const FString& Key)
 		return false;
 	}
 
-	FName KeyName(*Key);
-	FKey ParsedKey(KeyName);
+	// Parse modifier prefixes: "Ctrl+Shift+Alt+Meta+KeyName"
+	bool bShift = false;
+	bool bControl = false;
+	bool bAlt = false;
+	bool bCommand = false;
+	FString Remaining = Key;
+
+	while (true)
+	{
+		if (Remaining.StartsWith(TEXT("Shift+"), ESearchCase::IgnoreCase))
+		{
+			bShift = true;
+			Remaining.RightChopInline(6);
+		}
+		else if (Remaining.StartsWith(TEXT("Ctrl+"), ESearchCase::IgnoreCase) || Remaining.StartsWith(TEXT("Control+"), ESearchCase::IgnoreCase))
+		{
+			bControl = true;
+			Remaining.RightChopInline(Remaining.StartsWith(TEXT("Ctrl+"), ESearchCase::IgnoreCase) ? 5 : 8);
+		}
+		else if (Remaining.StartsWith(TEXT("Alt+"), ESearchCase::IgnoreCase))
+		{
+			bAlt = true;
+			Remaining.RightChopInline(4);
+		}
+		else if (Remaining.StartsWith(TEXT("Meta+"), ESearchCase::IgnoreCase) || Remaining.StartsWith(TEXT("Cmd+"), ESearchCase::IgnoreCase))
+		{
+			bCommand = true;
+			Remaining.RightChopInline(Remaining.StartsWith(TEXT("Meta+"), ESearchCase::IgnoreCase) ? 5 : 4);
+		}
+		else
+		{
+			break;
+		}
+	}
+
+	FKey ParsedKey(*Remaining);
 
 	if (!ParsedKey.IsValid())
 	{
@@ -276,9 +420,31 @@ bool USlateInspectorToolset::PressKey(const FString& Key)
 		return false;
 	}
 
-	FModifierKeysState ModifierKeys;
+	FModifierKeysState ModifierKeys(bShift, bShift, bControl, bControl, bAlt, bAlt, bCommand, bCommand, false);
 	FKeyEvent KeyDownEvent(ParsedKey, ModifierKeys, 0, false, 0, 0);
 	FSlateApplication::Get().ProcessKeyDownEvent(KeyDownEvent);
+
+	// Some keys are handled via OnKeyChar rather than OnKeyDown (e.g.
+	// BackSpace in SEditableText). Normally the OS translates WM_KEYDOWN
+	// into WM_CHAR, but since we bypass the OS message loop we must send
+	// the character event ourselves. Only send for unmodified keys (Ctrl/Alt
+	// combos are handled purely via OnKeyDown).
+	if (!bControl && !bAlt && !bCommand)
+	{
+		static const TMap<FKey, TCHAR> KeyToCharacter = {
+			{ EKeys::BackSpace, TCHAR('\b') },
+			{ EKeys::Tab,       TCHAR('\t') },
+			{ EKeys::Enter,     TCHAR('\r') },
+			{ EKeys::Escape,    TCHAR('\x1b') },
+		};
+
+		if (const TCHAR* CharCode = KeyToCharacter.Find(ParsedKey))
+		{
+			FCharacterEvent CharEvent(*CharCode, ModifierKeys, 0, false);
+			FSlateApplication::Get().ProcessKeyCharEvent(CharEvent);
+		}
+	}
+
 	FKeyEvent KeyUpEvent(ParsedKey, ModifierKeys, 0, false, 0, 0);
 	FSlateApplication::Get().ProcessKeyUpEvent(KeyUpEvent);
 
@@ -287,27 +453,53 @@ bool USlateInspectorToolset::PressKey(const FString& Key)
 
 bool USlateInspectorToolset::SelectOption(const FString& Ref, const FString& Value)
 {
+	if (!FSlateApplication::IsInitialized())
+	{
+		return false;
+	}
+
 	TSharedPtr<SWidget> Widget = ResolveRefOrWarn(Ref);
 	if (!Widget.IsValid())
 	{
 		return false;
 	}
 
+	// Capture the set of windows before opening, so we can identify the new dropdown popup.
+	TArray<TSharedRef<SWindow>> WindowsBefore;
+	FSlateApplication::Get().GetAllVisibleWindowsOrdered(WindowsBefore);
+
 	// Click to open the combobox dropdown.
 	if (!SimulateClick(Widget.ToSharedRef(), EKeys::LeftMouseButton, false))
 	{
 		return false;
 	}
 
-	// Search all visible windows for a text block matching the value.
-	if (!FSlateApplication::IsInitialized())
+	// Pump Slate twice: first tick creates the popup window, second tick
+	// lays out and materializes the list view items inside it.
+	// Guard against reentrancy if we're already inside a Slate tick.
+	if (!ensure(!FSlateApplication::Get().IsProcessingInput()))
 	{
+		UE_LOGFMT(LogSlateInspector, Warning, "SlateInspector: SelectOption called during input processing, cannot tick");
 		return false;
 	}
+	FSlateApplication::Get().Tick();
+	FSlateApplication::Get().Tick();
 
+	static const FName TextBlockType(TEXT("STextBlock"));
 	TSharedPtr<SWidget> MatchingTextBlock;
 	TArray<TSharedRef<SWindow>> VisibleWindows;
 	FSlateApplication::Get().GetAllVisibleWindowsOrdered(VisibleWindows);
+
+	// Search newly appeared windows first (the dropdown popup), then fall back
+	// to all windows. This avoids matching text in the combobox's own label.
+	TSet<TSharedRef<SWindow>> WindowsBeforeSet(WindowsBefore);
+	VisibleWindows.Sort([&WindowsBeforeSet](const TSharedRef<SWindow>& A, const TSharedRef<SWindow>& B)
+	{
+		bool bAIsNew = !WindowsBeforeSet.Contains(A);
+		bool bBIsNew = !WindowsBeforeSet.Contains(B);
+		return bAIsNew > bBIsNew;
+	});
+
 	for (const TSharedRef<SWindow>& Window : VisibleWindows)
 	{
 		TArray<TSharedRef<SWidget>> Queue;
@@ -317,11 +509,10 @@ bool USlateInspectorToolset::SelectOption(const FString& Ref, const FString& Val
 			Queue.Add(WindowChildren->GetChildAt(ChildIndex));
 		}
 
-		static const FName TextBlockType(TEXT("STextBlock"));
-		while (Queue.Num() > 0)
+		int32 QueueIndex = 0;
+		while (QueueIndex < Queue.Num())
 		{
-			TSharedRef<SWidget> Current = Queue[0];
-			Queue.RemoveAt(0);
+			TSharedRef<SWidget> Current = Queue[QueueIndex++];
 
 			if (Current->GetType() == TextBlockType)
 			{
@@ -355,16 +546,33 @@ bool USlateInspectorToolset::SelectOption(const FString& Ref, const FString& Val
 	return SimulateClick(MatchingTextBlock.ToSharedRef(), EKeys::LeftMouseButton, false);
 }
 
-bool USlateInspectorToolset::Drag(const FString& StartRef, const FString& EndRef)
+bool USlateInspectorToolset::Drag(const FString& StartRef, const FString& EndRef, const FSlateInspectorModifierKeys& Modifiers)
 {
+	if (!FSlateApplication::IsInitialized())
+	{
+		return false;
+	}
+
 	TSharedPtr<SWidget> StartWidget = ResolveRefOrWarn(StartRef);
 	TSharedPtr<SWidget> EndWidget = ResolveRefOrWarn(EndRef);
-	if (!StartWidget.IsValid() || !EndWidget.IsValid() || !FSlateApplication::IsInitialized())
+	if (!StartWidget.IsValid() || !EndWidget.IsValid())
 	{
 		return false;
 	}
 
-	TSharedPtr<SWindow> StartWindow = FSlateApplication::Get().FindWidgetWindow(StartWidget.ToSharedRef());
+	// Walk parent chain to find the containing SWindow (same approach as SimulateClick).
+	TSharedPtr<SWindow> StartWindow;
+	TSharedPtr<SWidget> WalkCurrent = StartWidget;
+	static const FName WindowType(TEXT("SWindow"));
+	while (WalkCurrent.IsValid())
+	{
+		if (WalkCurrent->GetType() == WindowType)
+		{
+			StartWindow = StaticCastSharedPtr<SWindow>(WalkCurrent);
+			break;
+		}
+		WalkCurrent = WalkCurrent->GetParentWidget();
+	}
 	if (!StartWindow.IsValid())
 	{
 		UE_LOGFMT(LogSlateInspector, Warning, "SlateInspector: Could not find window for drag source");
@@ -374,12 +582,17 @@ bool USlateInspectorToolset::Drag(const FString& StartRef, const FString& EndRef
 	FVector2D StartPosition = GetWidgetScreenCenter(StartWidget.ToSharedRef());
 	FVector2D EndPosition = GetWidgetScreenCenter(EndWidget.ToSharedRef());
 
-	FModifierKeysState ModifierKeys;
+	const FModifierKeysState ModifierKeys = Modifiers.ToModifierKeysState();
 	TSet<FKey> PressedButtons;
 	PressedButtons.Add(EKeys::LeftMouseButton);
 
-	// Move cursor to start and press.
+	// Move cursor to start, update hover tracking, then press.
 	FSlateApplication::Get().SetCursorPos(StartPosition);
+	FPointerEvent StartMoveEvent(
+		0, StartPosition, StartPosition,
+		TSet<FKey>(), EKeys::Invalid, 0.0f, ModifierKeys);
+	FSlateApplication::Get().ProcessMouseMoveEvent(StartMoveEvent);
+
 	FPointerEvent MouseDownEvent(
 		0, StartPosition, StartPosition,
 		PressedButtons, EKeys::LeftMouseButton, 0.0f, ModifierKeys);
@@ -440,7 +653,7 @@ FString USlateInspectorToolset::Windows(const FString& Action, int32 Index)
 			FSlateApplication::Get().SetAllUserFocus(VisibleWindows[Index], EFocusCause::SetDirectly);
 			return TEXT("OK");
 		}
-		return TEXT("Error: Invalid window index");
+		return FString::Printf(TEXT("Error: Window index %d out of range (0..%d)"), Index, VisibleWindows.Num() - 1);
 	}
 
 	if (Action.Equals(TEXT("close"), ESearchCase::IgnoreCase))
@@ -450,7 +663,7 @@ FString USlateInspectorToolset::Windows(const FString& Action, int32 Index)
 			VisibleWindows[Index]->RequestDestroyWindow();
 			return TEXT("OK");
 		}
-		return TEXT("Error: Invalid window index");
+		return FString::Printf(TEXT("Error: Window index %d out of range (0..%d)"), Index, VisibleWindows.Num() - 1);
 	}
 
 	return TEXT("Error: Unknown action. Use 'list', 'select', or 'close'.");
@@ -473,9 +686,8 @@ bool USlateInspectorToolset::WaitFor(const FString& Text, const FString& TextGon
 		return true;
 	}
 
-	// Walk the widget tree looking for text WITHOUT rendering a snapshot,
-	// because Render() resets the ref cache and would silently invalidate
-	// refs obtained from a prior Snapshot() call.
+	// Walk the widget tree looking for text without rendering a full snapshot,
+	// which would be unnecessary overhead for a simple text presence check.
 	bool bFoundText = false;
 	bool bFoundTextGone = false;
 
@@ -491,7 +703,8 @@ bool USlateInspectorToolset::WaitFor(const FString& Text, const FString& TextGon
 		}
 
 		FString WidgetText;
-		if (Widget->GetType() == FName(TEXT("STextBlock")))
+		static const FName TextBlockType(TEXT("STextBlock"));
+		if (Widget->GetType() == TextBlockType)
 		{
 			WidgetText = StaticCastSharedRef<STextBlock>(Widget)->GetText().ToString();
 		}
@@ -550,7 +763,8 @@ bool USlateInspectorToolset::FillForm(const TArray<FSlateInspectorFormField>& Fi
 		{
 			// Click the checkbox if the desired state differs from the current state.
 			TSharedPtr<SWidget> Widget = ResolveRefOrWarn(Field.Ref);
-			if (Widget.IsValid() && Widget->GetType() == FName(TEXT("SCheckBox")))
+			static const FName CheckBoxType(TEXT("SCheckBox"));
+			if (Widget.IsValid() && Widget->GetType() == CheckBoxType)
 			{
 				TSharedRef<SCheckBox> CheckBox = StaticCastSharedRef<SCheckBox>(Widget.ToSharedRef());
 				bool bDesiredChecked = Field.Value.Equals(TEXT("true"), ESearchCase::IgnoreCase);
@@ -564,15 +778,15 @@ bool USlateInspectorToolset::FillForm(const TArray<FSlateInspectorFormField>& Fi
 					bFieldSuccess = true; // Already in desired state.
 				}
 			}
+			else if (Widget.IsValid())
+			{
+				UE_LOGFMT(LogSlateInspector, Warning, "SlateInspector: Widget '{Ref}' is not an SCheckBox", Field.Ref);
+			}
 		}
 		else if (Field.FieldType.Equals(TEXT("combobox"), ESearchCase::IgnoreCase))
 		{
 			bFieldSuccess = SelectOption(Field.Ref, Field.Value);
 		}
-		else if (Field.FieldType.Equals(TEXT("slider"), ESearchCase::IgnoreCase))
-		{
-			bFieldSuccess = Type(Field.Ref, Field.Value, true);
-		}
 		else
 		{
 			UE_LOGFMT(LogSlateInspector, Warning, "SlateInspector: Unknown field type '{FieldType}'", Field.FieldType);
diff --git a/Engine/Plugins/Experimental/SlateInspector/Source/SlateInspector/Public/SlateInspectorObserverManager.h b/Engine/Plugins/Experimental/SlateInspector/Source/SlateInspector/Public/SlateInspectorObserverManager.h
new file mode 100644
--- /dev/null
+++ b/Engine/Plugins/Experimental/SlateInspector/Source/SlateInspector/Public/SlateInspectorObserverManager.h
@@ -0,0 +1,96 @@
+// Copyright Epic Games, Inc. All Rights Reserved.
+
+#pragma once
+
+#include "CoreMinimal.h"
+#include "Widgets/SWidget.h"
+
+/** A single observer watching a Slate widget subtree. */
+struct FSlateInspectorObserver
+{
+	/** Unique identifier returned by AddObserver(). */
+	FString Identifier;
+
+	/** Root widget to observe. Invalid (null) when bRoot is true. */
+	TWeakPtr<SWidget> RootWidget;
+
+	/** When true, this is the root observer covering all visible windows (RootWidget is ignored). */
+	bool bRoot = false;
+
+	/** Maximum depth to walk from the root. */
+	int32 MaxDepth = 30;
+
+	/** Cached snapshot text from the most recent observer tick. */
+	FString CachedSnapshotText;
+};
+
+/**
+ * Manages a set of observers that continuously walk Slate widget subtrees
+ * on a throttled tick, keeping the ref cache up to date as the UI changes.
+ *
+ * A default root observer (all visible windows) is always present.
+ */
+class FSlateInspectorObserverManager
+{
+public:
+
+	SLATEINSPECTOR_API static FSlateInspectorObserverManager& Get();
+
+	/**
+	 * Register an observer on a widget subtree.
+	 *
+	 * @param RootWidget  The root of the subtree to observe.  Null = all visible windows.
+	 * @param MaxDepth    Maximum depth to walk.
+	 * @return The observer identifier (used with RemoveObserver and GetCachedSnapshot).
+	 */
+	SLATEINSPECTOR_API FString AddObserver(TSharedPtr<SWidget> RootWidget, int32 MaxDepth = 30);
+
+	/**
+	 * Remove an observer by identifier.
+	 *
+	 * @return True if the observer was found and removed.
+	 */
+	SLATEINSPECTOR_API bool RemoveObserver(const FString& Identifier);
+
+	/** Returns a copy of all active observers. */
+	SLATEINSPECTOR_API TArray<FSlateInspectorObserver> GetObservers() const;
+
+	/**
+	 * Returns the cached snapshot text for an observer.
+	 *
+	 * @return The most recently rendered snapshot text, or empty if not found.
+	 */
+	SLATEINSPECTOR_API FString GetCachedSnapshot(const FString& Identifier) const;
+
+	/**
+	 * Find the best matching observer for the given root and return its cached snapshot.
+	 *
+	 * A null RootWidget matches the root observer.  Otherwise, matches
+	 * an observer whose root is the given widget (preferring deeper MaxDepth).
+	 *
+	 * @param RootWidget  The root to match. Null = root observer.
+	 * @param MaxDepth    Only match observers with at least this depth.
+	 * @return The cached snapshot text, or empty if no matching observer found.
+	 */
+	SLATEINSPECTOR_API FString FindMatchingObserverSnapshot(TSharedPtr<SWidget> RootWidget, int32 MaxDepth) const;
+
+	/**
+	 * Called each Slate post-tick.  Throttled internally to ~100ms intervals.
+	 */
+	SLATEINSPECTOR_API void Tick(float DeltaTime);
+
+private:
+
+	TArray<FSlateInspectorObserver> Observers;
+
+	int32 NextIdentifier = 1;
+
+	/** Accumulated time since last observer pass. */
+	float AccumulatedTime = 0.0f;
+
+	/** Tick counter for periodic purge scheduling. */
+	int32 TickCounter = 0;
+
+	static constexpr float TickIntervalSeconds = 0.1f;
+	static constexpr int32 PurgeEveryNTicks = 10;
+};
diff --git a/Engine/Plugins/Experimental/SlateInspector/Source/SlateInspector/Public/SlateInspectorRefCache.h b/Engine/Plugins/Experimental/SlateInspector/Source/SlateInspector/Public/SlateInspectorRefCache.h
--- a/Engine/Plugins/Experimental/SlateInspector/Source/SlateInspector/Public/SlateInspectorRefCache.h
+++ b/Engine/Plugins/Experimental/SlateInspector/Source/SlateInspector/Public/SlateInspectorRefCache.h
@@ -9,36 +9,20 @@
  * Manages bidirectional mapping between SWidget instances and short reference
  * strings used by the SlateInspector snapshot and action tools.
  *
- * # Reset-per-snapshot model
+ * # Observer model (append-only)
  *
- * This cache follows the establised approach such as Playwright's accessibility snapshots:
- * every call to FSlateInspectorSnapshotRenderer::Render() calls Reset() first,
- * then the tree walk assigns fresh refs from scratch (b1, b2, tb1, ...).
- * Counters restart at 1 each time, so refs are small, dense, and predictable.
+ * When observers are active, refs are never reset.  New widgets encountered
+ * during observer ticks receive the next available counter value for their
+ * role prefix.  Existing widgets keep their refs indefinitely (the same
+ * SWidget pointer always maps to the same ref string).  Destroyed widgets
+ * are cleaned up periodically via PurgeExpired().
  *
- * A ref is valid from the moment a snapshot is taken until the next snapshot
- * replaces it.  The expected workflow is:
+ * Refs are no longer dense after widgets are destroyed (you may see b1, b3, b7
+ * instead of b1, b2, b3), but keeping refs alive for existing widgets
+ * matters more than compactness.
  *
- *   1. Take a snapshot  -> refs are assigned (b1, b2, tb1, ...)
- *   2. Perform actions   -> Click("b2"), Type("tb1", "hello")
- *   3. Take a new snapshot -> all previous refs are invalidated, new ones assigned
- *
- * This means callers must NOT cache refs across snapshots. Each snapshot is a
- * self-contained view of the widget tree with its own ref namespace.
- *
- * # Rational
- *
- * - Compact refs: counters restart at 1 each pass, so refs stay small (b1 not b847).
- * - Minimal stale refs: the full rebuild removes entries orphaned by prior
- *   snapshots.  Between snapshots, destroyed widgets may leave expired weak
- *   pointers — ResolveRef() returns nullptr for those.
- * - Deterministic ordering: identical widget trees produce identical ref assignments,
- *   making snapshots diff-friendly and easier for LLMs to reason about.
- * - Simplicity: no pruning heuristics, no garbage collection, no ref-reuse logic.
- *
- * The trade-off is that refs are ephemeral: they only survive until the next
- * snapshot. Expectation is that the caller always takes a fresh snapshot before acting,
- * so refs never need long-term stability.
+ * Reset() is available for callers that need to clear the cache entirely
+ * and start fresh (e.g. testing).
  */
 class FSlateInspectorRefCache
 {
@@ -68,14 +52,27 @@ public:
 	SLATEINSPECTOR_API TSharedPtr<SWidget> ResolveRef(const FString& Ref) const;
 
 	/**
-	 * Clears all mappings and resets counters to zero.
+	 * Looks up the ref for a widget without assigning one.
 	 *
-	 * Called at the start of each snapshot pass so that refs are assigned
-	 * fresh from scratch. After Reset(), the next button encountered will
-	 * be b1, the next textbox tb1, etc.
+	 * @return The ref string, or empty if the widget has no ref.
+	 */
+	SLATEINSPECTOR_API FString FindRef(TSharedRef<SWidget> Widget) const;
+
+	/**
+	 * Clears all mappings and resets counters to zero.
+	 * Observer-driven passes should NOT call this because they rely on
+	 * the append-only behavior to preserve existing refs.
 	 */
 	SLATEINSPECTOR_API void Reset();
 
+	/**
+	 * Removes entries whose widget has been destroyed (weak pointer expired).
+	 *
+	 * Called periodically by the observer manager to prevent unbounded
+	 * map growth.  Safe to call at any time.
+	 */
+	SLATEINSPECTOR_API void PurgeExpired();
+
 private:
 
 	/** Widget -> ref string (deduplication within a single snapshot pass). */
diff --git a/Engine/Plugins/Experimental/SlateInspector/Source/SlateInspector/Public/SlateInspectorSnapshotRenderer.h b/Engine/Plugins/Experimental/SlateInspector/Source/SlateInspector/Public/SlateInspectorSnapshotRenderer.h
--- a/Engine/Plugins/Experimental/SlateInspector/Source/SlateInspector/Public/SlateInspectorSnapshotRenderer.h
+++ b/Engine/Plugins/Experimental/SlateInspector/Source/SlateInspector/Public/SlateInspectorSnapshotRenderer.h
@@ -10,7 +10,7 @@
  *
  * Walks the tree depth-first, emitting one line per meaningful widget with
  * role, label, state flags, and a compact ref identifier.  Structural
- * containers (SBox, SBorder, SOverlay, etc.) are elided — their children
+ * containers (SBox, SBorder, SOverlay, etc.) are elided and their children
  * are emitted at the parent's indent level.
  */
 class FSlateInspectorSnapshotRenderer
@@ -23,9 +23,11 @@ public:
 	 * @param Root                     The widget subtree to snapshot.  If null, snapshots all top-level windows.
 	 * @param MaxDepth                 Maximum recursion depth (default 30).
 	 * @param bIncludeSourceLocations  Include [src=File:Line] tags showing where each widget was created in C++.
+	 * @param bResetCache              When true (default), clears the ref cache before rendering.
+	 *                                 Pass false for observer-driven passes that preserve existing refs.
 	 * @return Indented text representation of the widget tree.
 	 */
-	SLATEINSPECTOR_API static FString Render(TSharedPtr<SWidget> Root, int32 MaxDepth = 30, bool bIncludeSourceLocations = false);
+	SLATEINSPECTOR_API static FString Render(TSharedPtr<SWidget> Root, int32 MaxDepth = 30, bool bIncludeSourceLocations = false, bool bResetCache = true);
 
 	/** Register a custom widget type so it appears in snapshots. */
 	SLATEINSPECTOR_API static void RegisterWidgetRole(const FName& WidgetType, const FString& Role, const FString& RefPrefix);
diff --git a/Engine/Plugins/Experimental/SlateInspector/Source/SlateInspector/Public/SlateInspectorSubsystem.h b/Engine/Plugins/Experimental/SlateInspector/Source/SlateInspector/Public/SlateInspectorSubsystem.h
--- a/Engine/Plugins/Experimental/SlateInspector/Source/SlateInspector/Public/SlateInspectorSubsystem.h
+++ b/Engine/Plugins/Experimental/SlateInspector/Source/SlateInspector/Public/SlateInspectorSubsystem.h
@@ -24,5 +24,12 @@ public:
 
 private:
 
+	/** Forwarded from FSlateApplication::OnPostTick to drive observer updates. */
+	void OnSlatePostTick(float DeltaTime);
+
 	bool bToolsetRegistered = false;
+
+	FDelegateHandle PostTickDelegateHandle;
+
+	FString RootObserverIdentifier;
 };
diff --git a/Engine/Plugins/Experimental/SlateInspector/Source/SlateInspector/Public/SlateInspectorToolset.h b/Engine/Plugins/Experimental/SlateInspector/Source/SlateInspector/Public/SlateInspectorToolset.h
--- a/Engine/Plugins/Experimental/SlateInspector/Source/SlateInspector/Public/SlateInspectorToolset.h
+++ b/Engine/Plugins/Experimental/SlateInspector/Source/SlateInspector/Public/SlateInspectorToolset.h
@@ -2,12 +2,38 @@
 
 #pragma once
 
+#include "Input/Events.h"
 #include "ToolsetRegistry/UToolsetRegistry.h"
 #include "ToolsetRegistry/ToolsetImage.h"
 #include "Widgets/SWidget.h"
 
 #include "SlateInspectorToolset.generated.h"
 
+/** Modifier keys held during an input simulation (click, drag, etc.). */
+USTRUCT(BlueprintType)
+struct FSlateInspectorModifierKeys
+{
+	GENERATED_BODY()
+
+	UPROPERTY(BlueprintReadWrite, Category = "SlateInspector")
+	bool bShift = false;
+
+	UPROPERTY(BlueprintReadWrite, Category = "SlateInspector")
+	bool bCtrl = false;
+
+	UPROPERTY(BlueprintReadWrite, Category = "SlateInspector")
+	bool bAlt = false;
+
+	UPROPERTY(BlueprintReadWrite, Category = "SlateInspector")
+	bool bCmd = false;
+
+	/** Convert to Slate's FModifierKeysState. */
+	FModifierKeysState ToModifierKeysState() const
+	{
+		return FModifierKeysState(bShift, bShift, bCtrl, bCtrl, bAlt, bAlt, bCmd, bCmd, false);
+	}
+};
+
 /** Describes a single form field for the FillForm tool. */
 USTRUCT(BlueprintType)
 struct FSlateInspectorFormField
@@ -22,7 +48,7 @@ struct FSlateInspectorFormField
 	UPROPERTY(BlueprintReadWrite, Category = "SlateInspector")
 	FString Value;
 
-	/** The type of field: "textbox", "checkbox", "combobox", or "slider". */
+	/** The type of field: "textbox", "checkbox", or "combobox". */
 	UPROPERTY(BlueprintReadWrite, Category = "SlateInspector")
 	FString FieldType;
 };
@@ -34,10 +60,16 @@ struct FSlateInspectorFormField
  * Unreal Editor UI programmatically.  Registered via UToolsetRegistry
  * so the ModelContextProtocol plugin picks them up automatically.
  *
+ * A shallow root observer (depth 0) continuously tracks top-level
+ * windows.  Before working with a specific window or panel, call
+ * Observe() on it to get deep widget coverage, then Unobserve() when
+ * done.  Observers walk their subtree every ~100ms, assigning refs to
+ * newly appeared widgets and keeping the ref cache current.
+ *
  * Input simulation uses direct Slate event APIs (ProcessKeyCharEvent,
  * ProcessMouseButtonDownEvent, etc.) rather than the AutomationDriver,
- * because AutomationDriver's synchronous API seems to deadlock when called
- * from the game thread (which is where MCP tool calls execute).
+ * because AutomationDriver's synchronous API deadlocks when called from
+ * the game thread (which is where MCP tool calls execute).
  */
 UCLASS(BlueprintType, MinimalAPI)
 class USlateInspectorToolset : public UToolsetDefinition
@@ -46,84 +78,107 @@ class USlateInspectorToolset : public UToolsetDefinition
 
 public:
 
-	/** Capture a Slate UI accessibility snapshot. Use this BEFORE any other SlateInspector
-	 * tool -- it returns the refs needed by Click, Type, Hover, Screenshot, etc.
-	 * IMPORTANT: Every Snapshot call resets ALL refs. Never reuse refs from a prior Snapshot.
-	 * @param Ref                      Subtree root ref (from most recent Snapshot). Empty = all windows.
+	/** Capture a Slate UI accessibility snapshot.  Use this to read the current
+	 * widget tree and discover refs for action tools (Click, Type, Hover, etc.).
+	 * A shallow root observer (depth 0) covers top-level windows automatically.
+	 * Before interacting with a specific window or panel, call Observe() on it
+	 * to get deep coverage, then Snapshot that subtree to see its contents.
+	 * Refs discovered by a previous Snapshot remain usable. You do NOT need to
+	 * call Snapshot again before every action.
+	 * @param Ref                      Subtree root ref. Empty = all windows.
 	 * @param MaxDepth                 Maximum depth (default 30).
 	 * @param bIncludeSourceLocations  Include [src=File:Line] tags showing where each widget was created in C++. */
 	UFUNCTION(BlueprintCallable, Category = "SlateInspector")
-	static FString Snapshot(const FString& Ref, int32 MaxDepth = 30, bool bIncludeSourceLocations = false);
+	static SLATEINSPECTOR_API FString Snapshot(const FString& Ref, int32 MaxDepth = 30, bool bIncludeSourceLocations = false);
+
+	/** Register an observer on a widget subtree so its refs are continuously
+	 * kept up to date (~100ms tick).  Call this on the window or panel you are
+	 * about to work with. It ensures new widgets appearing in that subtree
+	 * are assigned refs automatically.  Unobserve when you are done.
+	 * A shallow root observer (depth 0) already covers top-level windows.
+	 * @param Ref       Root widget ref to observe. Empty = all visible windows.
+	 * @param MaxDepth  Maximum depth to walk from the root. */
+	UFUNCTION(BlueprintCallable, Category = "SlateInspector")
+	static SLATEINSPECTOR_API FString Observe(const FString& Ref, int32 MaxDepth = 30);
+
+	/** Remove an observer by its identifier.
+	 * @param Identifier  The identifier returned by Observe(). */
+	UFUNCTION(BlueprintCallable, Category = "SlateInspector")
+	static SLATEINSPECTOR_API bool Unobserve(const FString& Identifier);
+
+	/** List all active observers as a JSON array for debugging.
+	 * Each entry includes the observer identifier, whether it is the root observer,
+	 * the root widget ref (if any), max depth, and cached snapshot size. */
+	UFUNCTION(BlueprintCallable, Category = "SlateInspector")
+	static SLATEINSPECTOR_API FString ListObservers();
 
 	/** Screenshot a Slate widget or the active editor window. Prefer this over
 	 * SceneTools.take_screenshot for Editor UI; use SceneTools only for 3D viewport.
-	 * @param Ref  Slate widget ref (from most recent Snapshot). Empty = active window. */
+	 * @param Ref  Slate widget ref. Empty = active window. */
 	UFUNCTION(BlueprintCallable, Category = "SlateInspector")
-	static FToolsetImage Screenshot(const FString& Ref);
+	static SLATEINSPECTOR_API FToolsetImage Screenshot(const FString& Ref);
 
-	// -- Action tools (all refs must come from the most recent Snapshot) --
+	// -- Action tools --
 
-	/** Click a Slate widget identified by its snapshot ref.
-	 * Refs reset each Snapshot -- never reuse old refs.
-	 * @param Ref         Slate widget ref from the most recent Snapshot.
+	/** Click a Slate widget identified by its ref.
+	 * @param Ref         Slate widget ref.
 	 * @param Button      "left", "right", or "middle".
-	 * @param DoubleClick True for double-click. */
+	 * @param DoubleClick True for double-click.
+	 * @param Modifiers   Modifier keys held during the click. */
 	UFUNCTION(BlueprintCallable, Category = "SlateInspector")
-	static bool Click(const FString& Ref, const FString& Button = TEXT("left"), bool DoubleClick = false);
+	static SLATEINSPECTOR_API bool Click(const FString& Ref, const FString& Button = TEXT("left"), bool DoubleClick = false, const FSlateInspectorModifierKeys& Modifiers = FSlateInspectorModifierKeys());
 
 	/** Hover over a Slate widget, triggering any hover state or tooltip.
-	 * Refs reset each Snapshot -- never reuse old refs.
-	 * @param Ref  Slate widget ref from the most recent Snapshot. */
+	 * @param Ref  Slate widget ref. */
 	UFUNCTION(BlueprintCallable, Category = "SlateInspector")
-	static bool Hover(const FString& Ref);
+	static SLATEINSPECTOR_API bool Hover(const FString& Ref);
 
 	/** Type text into a Slate text input widget. Focuses the widget first,
 	 * then sends one key event per character.
-	 * Refs reset each Snapshot -- never reuse old refs.
-	 * @param Ref    Slate textbox ref from the most recent Snapshot.
+	 * @param Ref    Slate textbox ref.
 	 * @param Text   Text to type.
 	 * @param Submit Press Enter after typing. */
 	UFUNCTION(BlueprintCallable, Category = "SlateInspector")
-	static bool Type(const FString& Ref, const FString& Text, bool Submit = false);
+	static SLATEINSPECTOR_API bool Type(const FString& Ref, const FString& Text, bool Submit = false);
 
 	/** Press and release a keyboard key on the currently focused Slate widget.
-	 * @param Key  Key name, e.g. "Enter", "Escape", "Tab", "ArrowLeft". */
+	 * Supports modifier prefixes: "Ctrl+C", "Shift+1".
+	 * @param Key  Key name with optional modifiers, e.g. "Enter", "Ctrl+A", "Shift+Tab". */
 	UFUNCTION(BlueprintCallable, Category = "SlateInspector")
-	static bool PressKey(const FString& Key);
+	static SLATEINSPECTOR_API bool PressKey(const FString& Key);
 
 	/** Select an option in a Slate combobox by its text label.
 	 * Opens the dropdown, finds the matching text, and clicks it.
-	 * Refs reset each Snapshot -- never reuse old refs.
-	 * @param Ref    Slate combobox ref from the most recent Snapshot.
+	 * @param Ref    Slate combobox ref.
 	 * @param Value  Exact option text to select. */
 	UFUNCTION(BlueprintCallable, Category = "SlateInspector")
-	static bool SelectOption(const FString& Ref, const FString& Value);
+	static SLATEINSPECTOR_API bool SelectOption(const FString& Ref, const FString& Value);
 
 	/** Drag from one Slate widget to another (mouse down, move, release).
-	 * Refs reset each Snapshot -- never reuse old refs.
-	 * @param StartRef  Slate widget ref for the drag source.
-	 * @param EndRef    Slate widget ref for the drop target. */
+	 * @param StartRef   Slate widget ref for the drag source.
+	 * @param EndRef     Slate widget ref for the drop target.
+	 * @param Modifiers  Modifier keys held during the drag. */
 	UFUNCTION(BlueprintCallable, Category = "SlateInspector")
-	static bool Drag(const FString& StartRef, const FString& EndRef);
+	static SLATEINSPECTOR_API bool Drag(const FString& StartRef, const FString& EndRef, const FSlateInspectorModifierKeys& Modifiers = FSlateInspectorModifierKeys());
 
 	/** List, select, or close top-level Slate editor windows.
 	 * @param Action  "list" returns JSON array, "select" brings to front, "close" destroys.
 	 * @param Index   Window index for select/close. */
 	UFUNCTION(BlueprintCallable, Category = "SlateInspector")
-	static FString Windows(const FString& Action = TEXT("list"), int32 Index = -1);
+	static SLATEINSPECTOR_API FString Windows(const FString& Action = TEXT("list"), int32 Index = -1);
 
 	/** Check if text is present or absent in the Slate widget tree.
 	 * Non-blocking: checks once and returns immediately. Poll to wait.
 	 * @param Text       Text that must be present (empty = skip).
 	 * @param TextGone   Text that must be absent (empty = skip). */
 	UFUNCTION(BlueprintCallable, Category = "SlateInspector")
-	static bool WaitFor(const FString& Text = TEXT(""), const FString& TextGone = TEXT(""));
+	static SLATEINSPECTOR_API bool WaitFor(const FString& Text = TEXT(""), const FString& TextGone = TEXT(""));
 
-	/** Fill multiple Slate form fields at once. Refs from most recent Snapshot only.
+	/** Fill multiple Slate form fields at once.
 	 * @param Fields  Array of {Ref, Value, FieldType} where FieldType is
-	 *                "textbox", "checkbox", "combobox", or "slider". */
+	 *                "textbox", "checkbox", or "combobox". */
 	UFUNCTION(BlueprintCallable, Category = "SlateInspector")
-	static bool FillForm(const TArray<FSlateInspectorFormField>& Fields);
+	static SLATEINSPECTOR_API bool FillForm(const TArray<FSlateInspectorFormField>& Fields);
 
 private:
 
@@ -137,5 +192,5 @@ private:
 	static FKey ParseMouseButtonKey(const FString& Button);
 
 	/** Simulate a mouse click at a widget's screen-space center. */
-	static bool SimulateClick(TSharedRef<SWidget> Widget, FKey MouseButton, bool bDoubleClick);
+	static bool SimulateClick(TSharedRef<SWidget> Widget, FKey MouseButton, bool bDoubleClick, const FModifierKeysState& ModifierKeys = FModifierKeysState());
 };
diff --git a/Engine/Plugins/Experimental/SlateInspector/Source/SlateInspectorTests/Private/SlateInspectorTestPanel.cpp b/Engine/Plugins/Experimental/SlateInspector/Source/SlateInspectorTests/Private/SlateInspectorTestPanel.cpp
new file mode 100644
--- /dev/null
+++ b/Engine/Plugins/Experimental/SlateInspector/Source/SlateInspectorTests/Private/SlateInspectorTestPanel.cpp
@@ -0,0 +1,125 @@
+// Copyright Epic Games, Inc. All Rights Reserved.
+
+#include "SlateInspectorTestPanel.h"
+
+#include "Widgets/Input/SButton.h"
+#include "Widgets/Layout/SBox.h"
+#include "Widgets/SBoxPanel.h"
+#include "Widgets/Text/STextBlock.h"
+
+void SSlateInspectorTestPanel::Construct(const FArguments& InArgs)
+{
+	ComboOptions.Add(MakeShared<FString>(TEXT("Alpha")));
+	ComboOptions.Add(MakeShared<FString>(TEXT("Beta")));
+	ComboOptions.Add(MakeShared<FString>(TEXT("Gamma")));
+	SelectedComboOption = ComboOptions[0];
+
+	ChildSlot
+	[
+		SNew(SVerticalBox)
+
+		// Primary button
+		+ SVerticalBox::Slot().AutoHeight().Padding(4)
+		[
+			SNew(SButton)
+			.Text(FText::FromString(TEXT("Test Button")))
+			.OnClicked_Lambda([this]()
+			{
+				++ButtonClickCount;
+				return FReply::Handled();
+			})
+		]
+
+		// Second button (for double-click testing)
+		+ SVerticalBox::Slot().AutoHeight().Padding(4)
+		[
+			SNew(SButton)
+			.Text(FText::FromString(TEXT("Second Button")))
+			.OnClicked_Lambda([this]()
+			{
+				++SecondButtonClickCount;
+				return FReply::Handled();
+			})
+		]
+
+		// Single-line text box
+		+ SVerticalBox::Slot().AutoHeight().Padding(4)
+		[
+			SAssignNew(TextBox, SEditableTextBox)
+			.HintText(FText::FromString(TEXT("Type here")))
+		]
+
+		// Multi-line text box
+		+ SVerticalBox::Slot().AutoHeight().Padding(4)
+		[
+			SNew(SBox)
+			.HeightOverride(60)
+			[
+				SAssignNew(MultiLineTextBox, SMultiLineEditableTextBox)
+				.HintText(FText::FromString(TEXT("Multi-line")))
+			]
+		]
+
+		// Search box
+		+ SVerticalBox::Slot().AutoHeight().Padding(4)
+		[
+			SAssignNew(SearchBox, SSearchBox)
+		]
+
+		// Checkbox
+		+ SVerticalBox::Slot().AutoHeight().Padding(4)
+		[
+			SAssignNew(CheckBox, SCheckBox)
+		]
+
+		// Combobox
+		+ SVerticalBox::Slot().AutoHeight().Padding(4)
+		[
+			SAssignNew(ComboBox, SComboBox<TSharedPtr<FString>>)
+			.OptionsSource(&ComboOptions)
+			.OnSelectionChanged_Lambda([this](TSharedPtr<FString> NewValue, ESelectInfo::Type)
+			{
+				SelectedComboOption = NewValue;
+			})
+			.OnGenerateWidget_Lambda([](TSharedPtr<FString> Item)
+			{
+				return SNew(STextBlock).Text(FText::FromString(*Item));
+			})
+			[
+				SNew(STextBlock)
+				.Text_Lambda([this]() { return FText::FromString(*SelectedComboOption); })
+			]
+		]
+
+		// Slider
+		+ SVerticalBox::Slot().AutoHeight().Padding(4)
+		[
+			SAssignNew(Slider, SSlider)
+			.Value(0.0f)
+			.OnValueChanged_Lambda([this](float NewValue)
+			{
+				SliderValue = NewValue;
+			})
+		]
+
+		// SpinBox
+		+ SVerticalBox::Slot().AutoHeight().Padding(4)
+		[
+			SAssignNew(SpinBox, SSpinBox<float>)
+			.MinValue(0.0f)
+			.MaxValue(100.0f)
+			.Value(0.0f)
+			.OnValueChanged_Lambda([this](float NewValue)
+			{
+				SpinBoxValue = NewValue;
+			})
+		]
+
+		// Static text label
+		+ SVerticalBox::Slot().AutoHeight().Padding(4)
+		[
+			SNew(STextBlock)
+			.Text(FText::FromString(TEXT("Static Label")))
+		]
+	];
+}
diff --git a/Engine/Plugins/Experimental/SlateInspector/Source/SlateInspectorTests/Private/SlateInspectorTestPanel.h b/Engine/Plugins/Experimental/SlateInspector/Source/SlateInspectorTests/Private/SlateInspectorTestPanel.h
new file mode 100644
--- /dev/null
+++ b/Engine/Plugins/Experimental/SlateInspector/Source/SlateInspectorTests/Private/SlateInspectorTestPanel.h
@@ -0,0 +1,50 @@
+// Copyright Epic Games, Inc. All Rights Reserved.
+
+#pragma once
+
+#include "Widgets/SCompoundWidget.h"
+#include "Widgets/Input/SCheckBox.h"
+#include "Widgets/Input/SComboBox.h"
+#include "Widgets/Input/SEditableTextBox.h"
+#include "Widgets/Input/SSlider.h"
+#include "Widgets/Input/SSpinBox.h"
+#include "Widgets/Input/SSearchBox.h"
+#include "Widgets/Input/SMultiLineEditableTextBox.h"
+
+/**
+ * A widget panel containing one of each common widget type.
+ * Used by SlateInspector automation tests to exercise snapshot, observer,
+ * and input simulation tools against known widgets.
+ */
+class SSlateInspectorTestPanel : public SCompoundWidget
+{
+public:
+
+	SLATE_BEGIN_ARGS(SSlateInspectorTestPanel) {}
+	SLATE_END_ARGS()
+
+	void Construct(const FArguments& InArgs);
+
+	// Button
+	int32 ButtonClickCount = 0;
+	int32 SecondButtonClickCount = 0;
+
+	// Text input
+	TSharedPtr<SEditableTextBox> TextBox;
+	TSharedPtr<SMultiLineEditableTextBox> MultiLineTextBox;
+	TSharedPtr<SSearchBox> SearchBox;
+
+	// Toggle
+	TSharedPtr<SCheckBox> CheckBox;
+
+	// Combobox
+	TSharedPtr<SComboBox<TSharedPtr<FString>>> ComboBox;
+	TArray<TSharedPtr<FString>> ComboOptions;
+	TSharedPtr<FString> SelectedComboOption;
+
+	// Slider / SpinBox
+	TSharedPtr<SSlider> Slider;
+	float SliderValue = 0.0f;
+	TSharedPtr<SSpinBox<float>> SpinBox;
+	float SpinBoxValue = 0.0f;
+};
diff --git a/Engine/Plugins/Experimental/SlateInspector/Source/SlateInspectorTests/Private/SlateInspectorTests.cpp b/Engine/Plugins/Experimental/SlateInspector/Source/SlateInspectorTests/Private/SlateInspectorTests.cpp
new file mode 100644
--- /dev/null
+++ b/Engine/Plugins/Experimental/SlateInspector/Source/SlateInspectorTests/Private/SlateInspectorTests.cpp
@@ -0,0 +1,776 @@
+// Copyright Epic Games, Inc. All Rights Reserved.
+
+#include "CQTest.h"
+#include "Components/CQTestSlateComponent.h"
+#include "SlateInspectorTestPanel.h"
+
+#include "Framework/Application/SlateApplication.h"
+#include "SlateInspectorObserverManager.h"
+#include "SlateInspectorRefCache.h"
+#include "SlateInspectorToolset.h"
+#include "Widgets/SWindow.h"
+
+TEST_CLASS(SlateInspectorTest, "Editor.SlateInspector")
+{
+	TUniquePtr<FCQTestSlateComponent> SlateComponent;
+	TSharedPtr<SWindow> TestWindow;
+	TSharedPtr<SSlateInspectorTestPanel> TestPanel;
+
+	BEFORE_EACH()
+	{
+		SlateComponent = MakeUnique<FCQTestSlateComponent>();
+
+		TestPanel = SNew(SSlateInspectorTestPanel);
+
+		TestWindow = SNew(SWindow)
+			.Title(FText::FromString(TEXT("SlateInspector Test Window")))
+			.ScreenPosition(FVector2D(100, 100))
+			.ClientSize(FVector2D(400, 600))
+			.SupportsMaximize(false)
+			.SupportsMinimize(false)
+			.FocusWhenFirstShown(true)
+			[
+				TestPanel.ToSharedRef()
+			];
+
+		FSlateApplication::Get().AddWindow(TestWindow.ToSharedRef());
+		TestWindow->MoveWindowTo(FVector2D(100, 100));
+		TestWindow->BringToFront();
+
+		// Force layout so widgets have valid geometry before any interaction.
+		FSlateApplication::Get().Tick();
+	}
+
+	AFTER_EACH()
+	{
+		// Clean up any observers created during the test (except the root observer).
+		for (const FSlateInspectorObserver& Observer : FSlateInspectorObserverManager::Get().GetObservers())
+		{
+			if (!Observer.bRoot)
+			{
+				FSlateInspectorObserverManager::Get().RemoveObserver(Observer.Identifier);
+			}
+		}
+
+		if (TestWindow.IsValid())
+		{
+			TestWindow->RequestDestroyWindow();
+			TestWindow.Reset();
+		}
+		TestPanel.Reset();
+		SlateComponent.Reset();
+	}
+
+	// ── Snapshot tests ──────────────────────────────────────────────
+
+	TEST_METHOD(Snapshot_ReturnsWidgetTree)
+	{
+		const FString Snapshot = SnapshotTestWindow();
+		ASSERT_THAT(IsTrue(!Snapshot.IsEmpty()));
+		ASSERT_THAT(IsTrue(Snapshot.Contains(TEXT("button"))));
+		ASSERT_THAT(IsTrue(Snapshot.Contains(TEXT("textbox"))));
+		ASSERT_THAT(IsTrue(Snapshot.Contains(TEXT("checkbox"))));
+		ASSERT_THAT(IsTrue(Snapshot.Contains(TEXT("combobox"))));
+		ASSERT_THAT(IsTrue(Snapshot.Contains(TEXT("slider"))));
+	}
+
+	TEST_METHOD(Snapshot_SubtreeByRef)
+	{
+		const FString FullSnapshot = USlateInspectorToolset::Snapshot(TEXT(""), 30);
+		ASSERT_THAT(IsTrue(FullSnapshot.Contains(TEXT("SlateInspector Test Window"))));
+
+		const FString WindowRef = FindRefByRole(FullSnapshot, TEXT("window"), TEXT("SlateInspector Test Window"));
+		ASSERT_THAT(IsTrue(!WindowRef.IsEmpty()));
+
+		const FString SubtreeSnapshot = USlateInspectorToolset::Snapshot(WindowRef, 30);
+		ASSERT_THAT(IsTrue(SubtreeSnapshot.Contains(TEXT("button"))));
+	}
+
+	TEST_METHOD(Snapshot_ContainsStaticLabel)
+	{
+		const FString Snapshot = SnapshotTestWindow();
+		ASSERT_THAT(IsTrue(Snapshot.Contains(TEXT("Static Label"))));
+	}
+
+	TEST_METHOD(Snapshot_ContainsSpinBox)
+	{
+		const FString Snapshot = SnapshotTestWindow();
+		// SSpinBox<> maps to "slider" role via template prefix matching.
+		ASSERT_THAT(IsTrue(Snapshot.Contains(TEXT("slider"))));
+	}
+
+	// ── Observer tests ──────────────────────────────────────────────
+
+	TEST_METHOD(Observe_CreatesObserver)
+	{
+		const FString Snapshot = SnapshotTestWindow();
+		const FString WindowRef = FindRefByRole(Snapshot, TEXT("window"), TEXT("SlateInspector Test Window"));
+		ASSERT_THAT(IsTrue(!WindowRef.IsEmpty()));
+
+		const FString Identifier = USlateInspectorToolset::Observe(WindowRef, 30);
+		ASSERT_THAT(IsTrue(!Identifier.IsEmpty()));
+		ASSERT_THAT(IsTrue(!Identifier.StartsWith(TEXT("Error"))));
+
+		const FString ObserverList = USlateInspectorToolset::ListObservers();
+		ASSERT_THAT(IsTrue(ObserverList.Contains(Identifier)));
+
+		USlateInspectorToolset::Unobserve(Identifier);
+	}
+
+	TEST_METHOD(Observe_Deduplicates)
+	{
+		const FString Snapshot = SnapshotTestWindow();
+		const FString WindowRef = FindRefByRole(Snapshot, TEXT("window"), TEXT("SlateInspector Test Window"));
+
+		const FString Identifier1 = USlateInspectorToolset::Observe(WindowRef, 30);
+		const FString Identifier2 = USlateInspectorToolset::Observe(WindowRef, 30);
+		ASSERT_THAT(IsTrue(Identifier1 == Identifier2));
+
+		USlateInspectorToolset::Unobserve(Identifier1);
+	}
+
+	TEST_METHOD(Unobserve_RemovesObserver)
+	{
+		const FString Snapshot = SnapshotTestWindow();
+		const FString WindowRef = FindRefByRole(Snapshot, TEXT("window"), TEXT("SlateInspector Test Window"));
+		const FString Identifier = USlateInspectorToolset::Observe(WindowRef, 30);
+
+		const bool bRemoved = USlateInspectorToolset::Unobserve(Identifier);
+		ASSERT_THAT(IsTrue(bRemoved));
+
+		const FString ObserverList = USlateInspectorToolset::ListObservers();
+		ASSERT_THAT(IsTrue(!ObserverList.Contains(Identifier)));
+	}
+
+	TEST_METHOD(Unobserve_ReturnsFalseForUnknown)
+	{
+		const bool bRemoved = USlateInspectorToolset::Unobserve(TEXT("nonexistent_observer"));
+		ASSERT_THAT(IsTrue(!bRemoved));
+	}
+
+	TEST_METHOD(Observer_KeepsRefsAlive)
+	{
+		const FString Snapshot = SnapshotTestWindow();
+		const FString ButtonRef = FindRefByRole(Snapshot, TEXT("button"), TEXT("Test Button"));
+		ASSERT_THAT(IsTrue(!ButtonRef.IsEmpty()));
+
+		const FString WindowRef = FindRefByRole(Snapshot, TEXT("window"), TEXT("SlateInspector Test Window"));
+		const FString ObserverId = USlateInspectorToolset::Observe(WindowRef, 30);
+
+		TestCommandBuilder
+			.StartWhen([this]() { return SlateComponent->HaveTicksElapsed(15); })
+			.Then([this, ButtonRef, ObserverId]()
+			{
+				const bool bClicked = USlateInspectorToolset::Click(ButtonRef);
+				ASSERT_THAT(IsTrue(bClicked));
+				ASSERT_THAT(IsTrue(TestPanel->ButtonClickCount > 0));
+				USlateInspectorToolset::Unobserve(ObserverId);
+			});
+	}
+
+	// ── Click tests ─────────────────────────────────────────────────
+
+	TEST_METHOD(Click_PressesButton)
+	{
+		TestCommandBuilder
+			.StartWhen([this]() { return SlateComponent->HaveTicksElapsed(3); })
+			.Then([this]()
+			{
+				const FString Snapshot = SnapshotTestWindow();
+				const FString ButtonRef = FindRefByRole(Snapshot, TEXT("button"), TEXT("Test Button"));
+				ASSERT_THAT(IsTrue(!ButtonRef.IsEmpty()));
+
+				ASSERT_THAT(IsTrue(TestPanel->ButtonClickCount == 0));
+				const bool bClicked = USlateInspectorToolset::Click(ButtonRef);
+				ASSERT_THAT(IsTrue(bClicked));
+				ASSERT_THAT(IsTrue(TestPanel->ButtonClickCount == 1));
+			});
+	}
+
+	TEST_METHOD(Click_MultipleButtons)
+	{
+		TestCommandBuilder
+			.StartWhen([this]() { return SlateComponent->HaveTicksElapsed(3); })
+			.Then([this]()
+			{
+				const FString Snapshot = SnapshotTestWindow();
+				const FString FirstRef = FindRefByRole(Snapshot, TEXT("button"), TEXT("Test Button"));
+				const FString SecondRef = FindRefByRole(Snapshot, TEXT("button"), TEXT("Second Button"));
+				ASSERT_THAT(IsTrue(!FirstRef.IsEmpty()));
+				ASSERT_THAT(IsTrue(!SecondRef.IsEmpty()));
+				ASSERT_THAT(IsTrue(FirstRef != SecondRef));
+
+				USlateInspectorToolset::Click(FirstRef);
+				USlateInspectorToolset::Click(SecondRef);
+				USlateInspectorToolset::Click(FirstRef);
+
+				ASSERT_THAT(IsTrue(TestPanel->ButtonClickCount == 2));
+				ASSERT_THAT(IsTrue(TestPanel->SecondButtonClickCount == 1));
+			});
+	}
+
+	TEST_METHOD(Click_DoubleClick)
+	{
+		TestCommandBuilder
+			.StartWhen([this]() { return SlateComponent->HaveTicksElapsed(3); })
+			.Then([this]()
+			{
+				const FString Snapshot = SnapshotTestWindow();
+				const FString ButtonRef = FindRefByRole(Snapshot, TEXT("button"), TEXT("Test Button"));
+				ASSERT_THAT(IsTrue(!ButtonRef.IsEmpty()));
+
+				const bool bClicked = USlateInspectorToolset::Click(ButtonRef, TEXT("left"), true);
+				ASSERT_THAT(IsTrue(bClicked));
+			});
+	}
+
+	TEST_METHOD(Click_RightClick)
+	{
+		TestCommandBuilder
+			.StartWhen([this]() { return SlateComponent->HaveTicksElapsed(3); })
+			.Then([this]()
+			{
+				const FString Snapshot = SnapshotTestWindow();
+				const FString ButtonRef = FindRefByRole(Snapshot, TEXT("button"), TEXT("Test Button"));
+				ASSERT_THAT(IsTrue(!ButtonRef.IsEmpty()));
+
+				const bool bClicked = USlateInspectorToolset::Click(ButtonRef, TEXT("right"));
+				ASSERT_THAT(IsTrue(bClicked));
+			});
+	}
+
+	TEST_METHOD(Click_InvalidRef)
+	{
+		const bool bClicked = USlateInspectorToolset::Click(TEXT("nonexistent_ref"));
+		ASSERT_THAT(IsTrue(!bClicked));
+	}
+
+	TEST_METHOD(Click_WithShiftModifier)
+	{
+		TestCommandBuilder
+			.StartWhen([this]() { return SlateComponent->HaveTicksElapsed(3); })
+			.Then([this]()
+			{
+				const FString Snapshot = SnapshotTestWindow();
+				const FString ButtonRef = FindRefByRole(Snapshot, TEXT("button"), TEXT("Test Button"));
+				ASSERT_THAT(IsTrue(!ButtonRef.IsEmpty()));
+
+				FSlateInspectorModifierKeys Modifiers;
+				Modifiers.bShift = true;
+				const bool bClicked = USlateInspectorToolset::Click(ButtonRef, TEXT("left"), false, Modifiers);
+				ASSERT_THAT(IsTrue(bClicked));
+				ASSERT_THAT(IsTrue(TestPanel->ButtonClickCount == 1));
+			});
+	}
+
+	TEST_METHOD(Click_WithCtrlModifier)
+	{
+		TestCommandBuilder
+			.StartWhen([this]() { return SlateComponent->HaveTicksElapsed(3); })
+			.Then([this]()
+			{
+				const FString Snapshot = SnapshotTestWindow();
+				const FString ButtonRef = FindRefByRole(Snapshot, TEXT("button"), TEXT("Test Button"));
+				ASSERT_THAT(IsTrue(!ButtonRef.IsEmpty()));
+
+				FSlateInspectorModifierKeys Modifiers;
+				Modifiers.bCtrl = true;
+				const bool bClicked = USlateInspectorToolset::Click(ButtonRef, TEXT("left"), false, Modifiers);
+				ASSERT_THAT(IsTrue(bClicked));
+			});
+	}
+
+	TEST_METHOD(Click_WithMultipleModifiers)
+	{
+		TestCommandBuilder
+			.StartWhen([this]() { return SlateComponent->HaveTicksElapsed(3); })
+			.Then([this]()
+			{
+				const FString Snapshot = SnapshotTestWindow();
+				const FString ButtonRef = FindRefByRole(Snapshot, TEXT("button"), TEXT("Test Button"));
+				ASSERT_THAT(IsTrue(!ButtonRef.IsEmpty()));
+
+				FSlateInspectorModifierKeys Modifiers;
+				Modifiers.bShift = true;
+				Modifiers.bCtrl = true;
+				Modifiers.bAlt = true;
+				const bool bClicked = USlateInspectorToolset::Click(ButtonRef, TEXT("left"), false, Modifiers);
+				ASSERT_THAT(IsTrue(bClicked));
+			});
+	}
+
+	// ── Drag tests ─────────────────────────────────────────────────
+
+	TEST_METHOD(Drag_BetweenButtons)
+	{
+		TestCommandBuilder
+			.StartWhen([this]() { return SlateComponent->HaveTicksElapsed(3); })
+			.Then([this]()
+			{
+				const FString Snapshot = SnapshotTestWindow();
+				const FString FirstRef = FindRefByRole(Snapshot, TEXT("button"), TEXT("Test Button"));
+				const FString SecondRef = FindRefByRole(Snapshot, TEXT("button"), TEXT("Second Button"));
+				ASSERT_THAT(IsTrue(!FirstRef.IsEmpty()));
+				ASSERT_THAT(IsTrue(!SecondRef.IsEmpty()));
+
+				const bool bDragged = USlateInspectorToolset::Drag(FirstRef, SecondRef);
+				ASSERT_THAT(IsTrue(bDragged));
+			});
+	}
+
+	TEST_METHOD(Drag_WithModifiers)
+	{
+		TestCommandBuilder
+			.StartWhen([this]() { return SlateComponent->HaveTicksElapsed(3); })
+			.Then([this]()
+			{
+				const FString Snapshot = SnapshotTestWindow();
+				const FString FirstRef = FindRefByRole(Snapshot, TEXT("button"), TEXT("Test Button"));
+				const FString SecondRef = FindRefByRole(Snapshot, TEXT("button"), TEXT("Second Button"));
+
+				FSlateInspectorModifierKeys Modifiers;
+				Modifiers.bShift = true;
+				const bool bDragged = USlateInspectorToolset::Drag(FirstRef, SecondRef, Modifiers);
+				ASSERT_THAT(IsTrue(bDragged));
+			});
+	}
+
+	TEST_METHOD(Drag_InvalidStartRef)
+	{
+		const bool bDragged = USlateInspectorToolset::Drag(TEXT("nonexistent"), TEXT("nonexistent2"));
+		ASSERT_THAT(IsTrue(!bDragged));
+	}
+
+	TEST_METHOD(Drag_SliderChangesValue)
+	{
+		TestCommandBuilder
+			.StartWhen([this]() { return SlateComponent->HaveTicksElapsed(3); })
+			.Then([this]()
+			{
+				ASSERT_THAT(IsTrue(TestPanel->SliderValue == 0.0f));
+
+				const FString Snapshot = SnapshotTestWindow();
+				const TArray<FString> SliderRefs = FindAllRefsByRole(Snapshot, TEXT("slider"));
+				// First slider ref is the SSlider, second is the SSpinBox.
+				ASSERT_THAT(IsTrue(SliderRefs.Num() >= 1));
+				const FString SliderRef = SliderRefs[0];
+
+				// Drag from the slider to one of the buttons (to the right and up).
+				// This should move the slider thumb and change the value.
+				const FString ButtonRef = FindRefByRole(Snapshot, TEXT("button"), TEXT("Test Button"));
+				const bool bDragged = USlateInspectorToolset::Drag(SliderRef, ButtonRef);
+				ASSERT_THAT(IsTrue(bDragged));
+
+				// The slider value should have changed from its initial 0.0.
+				// We can't predict the exact value since it depends on geometry,
+				// but it should no longer be zero.
+				ASSERT_THAT(IsTrue(TestPanel->SliderValue != 0.0f));
+			});
+	}
+
+	// ── Checkbox tests ──────────────────────────────────────────────
+
+	TEST_METHOD(Click_TogglesCheckbox)
+	{
+		TestCommandBuilder
+			.StartWhen([this]() { return SlateComponent->HaveTicksElapsed(3); })
+			.Then([this]()
+			{
+				const FString Snapshot = SnapshotTestWindow();
+				const FString CheckBoxRef = FindRefByRole(Snapshot, TEXT("checkbox"), TEXT(""));
+				ASSERT_THAT(IsTrue(!CheckBoxRef.IsEmpty()));
+
+				ASSERT_THAT(IsTrue(!TestPanel->CheckBox->IsChecked()));
+
+				USlateInspectorToolset::Click(CheckBoxRef);
+				ASSERT_THAT(IsTrue(TestPanel->CheckBox->IsChecked()));
+
+				USlateInspectorToolset::Click(CheckBoxRef);
+				ASSERT_THAT(IsTrue(!TestPanel->CheckBox->IsChecked()));
+			});
+	}
+
+	// ── Type tests ──────────────────────────────────────────────────
+
+	TEST_METHOD(Type_EntersText)
+	{
+		const FString Snapshot = SnapshotTestWindow();
+		const FString TextBoxRef = FindRefByRole(Snapshot, TEXT("textbox"), TEXT(""));
+		ASSERT_THAT(IsTrue(!TextBoxRef.IsEmpty()));
+
+		const bool bTyped = USlateInspectorToolset::Type(TextBoxRef, TEXT("hello"));
+		ASSERT_THAT(IsTrue(bTyped));
+		ASSERT_THAT(IsTrue(TestPanel->TextBox->GetText().ToString() == TEXT("hello")));
+	}
+
+	TEST_METHOD(Type_AppendsToExistingText)
+	{
+		const FString Snapshot = SnapshotTestWindow();
+		const FString TextBoxRef = FindRefByRole(Snapshot, TEXT("textbox"), TEXT(""));
+
+		USlateInspectorToolset::Type(TextBoxRef, TEXT("foo"));
+		USlateInspectorToolset::Type(TextBoxRef, TEXT("bar"));
+		ASSERT_THAT(IsTrue(TestPanel->TextBox->GetText().ToString() == TEXT("foobar")));
+	}
+
+	TEST_METHOD(Type_SpecialCharacters)
+	{
+		const FString Snapshot = SnapshotTestWindow();
+		const FString TextBoxRef = FindRefByRole(Snapshot, TEXT("textbox"), TEXT(""));
+
+		USlateInspectorToolset::Type(TextBoxRef, TEXT("a@b#c$d"));
+		ASSERT_THAT(IsTrue(TestPanel->TextBox->GetText().ToString() == TEXT("a@b#c$d")));
+	}
+
+	TEST_METHOD(Type_WithSubmit)
+	{
+		const FString Snapshot = SnapshotTestWindow();
+		const FString TextBoxRef = FindRefByRole(Snapshot, TEXT("textbox"), TEXT(""));
+		ASSERT_THAT(IsTrue(!TextBoxRef.IsEmpty()));
+
+		const bool bTyped = USlateInspectorToolset::Type(TextBoxRef, TEXT("submitted"), true);
+		ASSERT_THAT(IsTrue(bTyped));
+	}
+
+	TEST_METHOD(Type_IntoSearchBox)
+	{
+		const FString Snapshot = SnapshotTestWindow();
+		// SSearchBox contains an SEditableTextBox internally, so look for a textbox.
+		// The search box is the third textbox in the panel (after TextBox and MultiLineTextBox).
+		const TArray<FString> TextBoxRefs = FindAllRefsByRole(Snapshot, TEXT("textbox"));
+		// We need at least 2 textbox refs (single-line and search box; multi-line may be a different role).
+		ASSERT_THAT(IsTrue(TextBoxRefs.Num() >= 2));
+
+		// Type into the last textbox (search box).
+		const FString SearchRef = TextBoxRefs.Last();
+		const bool bTyped = USlateInspectorToolset::Type(SearchRef, TEXT("search query"));
+		ASSERT_THAT(IsTrue(bTyped));
+	}
+
+	TEST_METHOD(Type_InvalidRef)
+	{
+		const bool bTyped = USlateInspectorToolset::Type(TEXT("nonexistent_ref"), TEXT("hello"));
+		ASSERT_THAT(IsTrue(!bTyped));
+	}
+
+	// ── PressKey in text context ────────────────────────────────────
+
+	TEST_METHOD(PressKey_BackspaceDeletesCharacter)
+	{
+		TestCommandBuilder
+			.StartWhen([this]() { return SlateComponent->HaveTicksElapsed(3); })
+			.Then([this]()
+			{
+				const FString Snapshot = SnapshotTestWindow();
+				const FString TextBoxRef = FindRefByRole(Snapshot, TEXT("textbox"), TEXT(""));
+
+				USlateInspectorToolset::Type(TextBoxRef, TEXT("abc"));
+				ASSERT_THAT(IsTrue(TestPanel->TextBox->GetText().ToString() == TEXT("abc")));
+
+				// BackSpace should delete the last character.
+				// Re-focus via Type to ensure the editable text has focus.
+				USlateInspectorToolset::Type(TextBoxRef, TEXT(""));
+				USlateInspectorToolset::PressKey(TEXT("BackSpace"));
+				ASSERT_THAT(IsTrue(TestPanel->TextBox->GetText().ToString() == TEXT("ab")));
+			});
+	}
+
+	TEST_METHOD(PressKey_SelectAllAndDelete)
+	{
+		TestCommandBuilder
+			.StartWhen([this]() { return SlateComponent->HaveTicksElapsed(3); })
+			.Then([this]()
+			{
+				const FString Snapshot = SnapshotTestWindow();
+				const FString TextBoxRef = FindRefByRole(Snapshot, TEXT("textbox"), TEXT(""));
+
+				USlateInspectorToolset::Type(TextBoxRef, TEXT("hello world"));
+				ASSERT_THAT(IsTrue(TestPanel->TextBox->GetText().ToString() == TEXT("hello world")));
+
+				// Ctrl+A to select all, then Delete to clear.
+				USlateInspectorToolset::PressKey(TEXT("Ctrl+A"));
+				USlateInspectorToolset::PressKey(TEXT("Delete"));
+				ASSERT_THAT(IsTrue(TestPanel->TextBox->GetText().ToString().IsEmpty()));
+			});
+	}
+
+	// ── Hover tests ─────────────────────────────────────────────────
+
+	TEST_METHOD(Hover_ReturnsTrue)
+	{
+		TestCommandBuilder
+			.StartWhen([this]() { return SlateComponent->HaveTicksElapsed(3); })
+			.Then([this]()
+			{
+				const FString Snapshot = SnapshotTestWindow();
+				const FString ButtonRef = FindRefByRole(Snapshot, TEXT("button"), TEXT("Test Button"));
+				ASSERT_THAT(IsTrue(!ButtonRef.IsEmpty()));
+
+				const bool bHovered = USlateInspectorToolset::Hover(ButtonRef);
+				ASSERT_THAT(IsTrue(bHovered));
+			});
+	}
+
+	TEST_METHOD(Hover_InvalidRef)
+	{
+		const bool bHovered = USlateInspectorToolset::Hover(TEXT("nonexistent_ref"));
+		ASSERT_THAT(IsTrue(!bHovered));
+	}
+
+	// ── PressKey tests ──────────────────────────────────────────────
+
+	TEST_METHOD(PressKey_Escape)
+	{
+		const bool bPressed = USlateInspectorToolset::PressKey(TEXT("Escape"));
+		ASSERT_THAT(IsTrue(bPressed));
+	}
+
+	TEST_METHOD(PressKey_Tab)
+	{
+		const bool bPressed = USlateInspectorToolset::PressKey(TEXT("Tab"));
+		ASSERT_THAT(IsTrue(bPressed));
+	}
+
+	// ── SelectOption tests ──────────────────────────────────────────
+
+	TEST_METHOD(SelectOption_ChangesCombobox)
+	{
+		TestCommandBuilder
+			.StartWhen([this]() { return SlateComponent->HaveTicksElapsed(3); })
+			.Then([this]()
+			{
+				ASSERT_THAT(IsTrue(*TestPanel->SelectedComboOption == TEXT("Alpha")));
+
+				const FString Snapshot = SnapshotTestWindow();
+				const FString ComboRef = FindRefByRole(Snapshot, TEXT("combobox"), TEXT(""));
+				ASSERT_THAT(IsTrue(!ComboRef.IsEmpty()));
+
+				const bool bSelected = USlateInspectorToolset::SelectOption(ComboRef, TEXT("Beta"));
+				ASSERT_THAT(IsTrue(bSelected));
+				ASSERT_THAT(IsTrue(*TestPanel->SelectedComboOption == TEXT("Beta")));
+			});
+	}
+
+	TEST_METHOD(SelectOption_InvalidOption)
+	{
+		TestCommandBuilder
+			.StartWhen([this]() { return SlateComponent->HaveTicksElapsed(3); })
+			.Then([this]()
+			{
+				const FString Snapshot = SnapshotTestWindow();
+				const FString ComboRef = FindRefByRole(Snapshot, TEXT("combobox"), TEXT(""));
+				ASSERT_THAT(IsTrue(!ComboRef.IsEmpty()));
+
+				const bool bSelected = USlateInspectorToolset::SelectOption(ComboRef, TEXT("Nonexistent"));
+				ASSERT_THAT(IsTrue(!bSelected));
+				// Original selection unchanged.
+				ASSERT_THAT(IsTrue(*TestPanel->SelectedComboOption == TEXT("Alpha")));
+			});
+	}
+
+	// ── FillForm tests ──────────────────────────────────────────────
+
+	TEST_METHOD(FillForm_TextBox)
+	{
+		const FString Snapshot = SnapshotTestWindow();
+		const FString TextBoxRef = FindRefByRole(Snapshot, TEXT("textbox"), TEXT(""));
+		ASSERT_THAT(IsTrue(!TextBoxRef.IsEmpty()));
+
+		TArray<FSlateInspectorFormField> Fields;
+		FSlateInspectorFormField Field;
+		Field.Ref = TextBoxRef;
+		Field.Value = TEXT("form value");
+		Field.FieldType = TEXT("textbox");
+		Fields.Add(Field);
+
+		const bool bFilled = USlateInspectorToolset::FillForm(Fields);
+		ASSERT_THAT(IsTrue(bFilled));
+		ASSERT_THAT(IsTrue(TestPanel->TextBox->GetText().ToString() == TEXT("form value")));
+	}
+
+	TEST_METHOD(FillForm_Checkbox)
+	{
+		TestCommandBuilder
+			.StartWhen([this]() { return SlateComponent->HaveTicksElapsed(3); })
+			.Then([this]()
+			{
+				const FString Snapshot = SnapshotTestWindow();
+				const FString CheckBoxRef = FindRefByRole(Snapshot, TEXT("checkbox"), TEXT(""));
+				ASSERT_THAT(IsTrue(!CheckBoxRef.IsEmpty()));
+				ASSERT_THAT(IsTrue(!TestPanel->CheckBox->IsChecked()));
+
+				TArray<FSlateInspectorFormField> Fields;
+				FSlateInspectorFormField Field;
+				Field.Ref = CheckBoxRef;
+				Field.Value = TEXT("true");
+				Field.FieldType = TEXT("checkbox");
+				Fields.Add(Field);
+
+				const bool bFilled = USlateInspectorToolset::FillForm(Fields);
+				ASSERT_THAT(IsTrue(bFilled));
+				ASSERT_THAT(IsTrue(TestPanel->CheckBox->IsChecked()));
+			});
+	}
+
+	TEST_METHOD(FillForm_CheckboxAlreadyInDesiredState)
+	{
+		TestCommandBuilder
+			.StartWhen([this]() { return SlateComponent->HaveTicksElapsed(3); })
+			.Then([this]()
+			{
+				const FString Snapshot = SnapshotTestWindow();
+				const FString CheckBoxRef = FindRefByRole(Snapshot, TEXT("checkbox"), TEXT(""));
+				ASSERT_THAT(IsTrue(!TestPanel->CheckBox->IsChecked()));
+
+				TArray<FSlateInspectorFormField> Fields;
+				FSlateInspectorFormField Field;
+				Field.Ref = CheckBoxRef;
+				Field.Value = TEXT("false");
+				Field.FieldType = TEXT("checkbox");
+				Fields.Add(Field);
+
+				// Should succeed without clicking (already unchecked).
+				const bool bFilled = USlateInspectorToolset::FillForm(Fields);
+				ASSERT_THAT(IsTrue(bFilled));
+				ASSERT_THAT(IsTrue(!TestPanel->CheckBox->IsChecked()));
+			});
+	}
+
+	// ── WaitFor tests ───────────────────────────────────────────────
+
+	TEST_METHOD(WaitFor_FindsExistingText)
+	{
+		const bool bFound = USlateInspectorToolset::WaitFor(TEXT("Static Label"));
+		ASSERT_THAT(IsTrue(bFound));
+	}
+
+	TEST_METHOD(WaitFor_DoesNotFindMissingText)
+	{
+		const bool bFound = USlateInspectorToolset::WaitFor(TEXT("This Text Does Not Exist"));
+		ASSERT_THAT(IsTrue(!bFound));
+	}
+
+	TEST_METHOD(WaitFor_TextGone)
+	{
+		const bool bGone = USlateInspectorToolset::WaitFor(TEXT(""), TEXT("This Text Does Not Exist"));
+		ASSERT_THAT(IsTrue(bGone));
+	}
+
+	TEST_METHOD(WaitFor_TextGoneFails)
+	{
+		const bool bGone = USlateInspectorToolset::WaitFor(TEXT(""), TEXT("Static Label"));
+		ASSERT_THAT(IsTrue(!bGone));
+	}
+
+	// ── Windows tests ───────────────────────────────────────────────
+
+	TEST_METHOD(Windows_ListContainsTestWindow)
+	{
+		const FString WindowList = USlateInspectorToolset::Windows(TEXT("list"));
+		ASSERT_THAT(IsTrue(WindowList.Contains(TEXT("SlateInspector Test Window"))));
+	}
+
+	// ── Ref cache tests ─────────────────────────────────────────────
+
+	TEST_METHOD(PurgeExpired_CleansDeadRefs)
+	{
+		TSharedRef<STextBlock> EphemeralWidget = SNew(STextBlock).Text(FText::FromString(TEXT("Ephemeral")));
+		const FString Ref = FSlateInspectorRefCache::Get().GetOrAssignRef(EphemeralWidget, TEXT("x"));
+		ASSERT_THAT(IsTrue(!Ref.IsEmpty()));
+
+		TSharedPtr<SWidget> Resolved = FSlateInspectorRefCache::Get().ResolveRef(Ref);
+		ASSERT_THAT(IsTrue(Resolved.IsValid()));
+
+		// Drop all references so the widget is destroyed.
+		Resolved.Reset();
+		EphemeralWidget = SNew(STextBlock);
+
+		FSlateInspectorRefCache::Get().PurgeExpired();
+		TSharedPtr<SWidget> ResolvedAfterPurge = FSlateInspectorRefCache::Get().ResolveRef(Ref);
+		ASSERT_THAT(IsTrue(!ResolvedAfterPurge.IsValid()));
+	}
+
+	TEST_METHOD(FindRef_ReturnsEmptyForUnknownWidget)
+	{
+		TSharedRef<STextBlock> Widget = SNew(STextBlock);
+		const FString Ref = FSlateInspectorRefCache::Get().FindRef(Widget);
+		ASSERT_THAT(IsTrue(Ref.IsEmpty()));
+	}
+
+	TEST_METHOD(ResolveRef_ReturnsNullForUnknownRef)
+	{
+		TSharedPtr<SWidget> Resolved = FSlateInspectorRefCache::Get().ResolveRef(TEXT("zzz999"));
+		ASSERT_THAT(IsTrue(!Resolved.IsValid()));
+	}
+
+	// ── Helpers ─────────────────────────────────────────────────────
+
+	/** Snapshot only the test window's subtree so we don't accidentally
+	 *  match editor widgets (e.g. editor checkboxes vs. our test checkbox). */
+	FString SnapshotTestWindow()
+	{
+		// First get the window ref from a full snapshot.
+		const FString FullSnapshot = USlateInspectorToolset::Snapshot(TEXT(""), 30);
+		const FString WindowRef = FindRefByRole(FullSnapshot, TEXT("window"), TEXT("SlateInspector Test Window"));
+		if (WindowRef.IsEmpty())
+		{
+			return FString();
+		}
+		return USlateInspectorToolset::Snapshot(WindowRef, 30);
+	}
+
+	FString FindRefByRole(const FString& Snapshot, const FString& Role, const FString& Label)
+	{
+		TArray<FString> Lines;
+		Snapshot.ParseIntoArrayLines(Lines);
+		for (const FString& Line : Lines)
+		{
+			FString Trimmed = Line.TrimStartAndEnd();
+			if (!Trimmed.StartsWith(Role))
+			{
+				continue;
+			}
+			if (!Label.IsEmpty() && !Trimmed.Contains(Label))
+			{
+				continue;
+			}
+			int32 RefStart = Trimmed.Find(TEXT("[ref="));
+			if (RefStart != INDEX_NONE)
+			{
+				RefStart += 5;
+				int32 RefEnd = Trimmed.Find(TEXT("]"), ESearchCase::IgnoreCase, ESearchDir::FromStart, RefStart);
+				if (RefEnd != INDEX_NONE)
+				{
+					return Trimmed.Mid(RefStart, RefEnd - RefStart);
+				}
+			}
+		}
+		return FString();
+	}
+
+	TArray<FString> FindAllRefsByRole(const FString& Snapshot, const FString& Role)
+	{
+		TArray<FString> Refs;
+		TArray<FString> Lines;
+		Snapshot.ParseIntoArrayLines(Lines);
+		for (const FString& Line : Lines)
+		{
+			FString Trimmed = Line.TrimStartAndEnd();
+			if (!Trimmed.StartsWith(Role))
+			{
+				continue;
+			}
+			int32 RefStart = Trimmed.Find(TEXT("[ref="));
+			if (RefStart != INDEX_NONE)
+			{
+				RefStart += 5;
+				int32 RefEnd = Trimmed.Find(TEXT("]"), ESearchCase::IgnoreCase, ESearchDir::FromStart, RefStart);
+				if (RefEnd != INDEX_NONE)
+				{
+					Refs.Add(Trimmed.Mid(RefStart, RefEnd - RefStart));
+				}
+			}
+		}
+		return Refs;
+	}
+};
diff --git a/Engine/Plugins/Experimental/SlateInspector/Source/SlateInspectorTests/Private/SlateInspectorTestsModule.cpp b/Engine/Plugins/Experimental/SlateInspector/Source/SlateInspectorTests/Private/SlateInspectorTestsModule.cpp
new file mode 100644
--- /dev/null
+++ b/Engine/Plugins/Experimental/SlateInspector/Source/SlateInspectorTests/Private/SlateInspectorTestsModule.cpp
@@ -0,0 +1,9 @@
+// Copyright Epic Games, Inc. All Rights Reserved.
+
+#include "Modules/ModuleManager.h"
+
+class FSlateInspectorTestsModule : public IModuleInterface
+{
+};
+
+IMPLEMENT_MODULE(FSlateInspectorTestsModule, SlateInspectorTests)
diff --git a/Engine/Plugins/Experimental/SlateInspector/Source/SlateInspectorTests/SlateInspectorTests.Build.cs b/Engine/Plugins/Experimental/SlateInspector/Source/SlateInspectorTests/SlateInspectorTests.Build.cs
new file mode 100644
--- /dev/null
+++ b/Engine/Plugins/Experimental/SlateInspector/Source/SlateInspectorTests/SlateInspectorTests.Build.cs
@@ -0,0 +1,23 @@
+// Copyright Epic Games, Inc. All Rights Reserved.
+
+using UnrealBuildTool;
+
+public class SlateInspectorTests : ModuleRules
+{
+	public SlateInspectorTests(ReadOnlyTargetRules Target) : base(Target)
+	{
+		PCHUsage = PCHUsageMode.UseExplicitOrSharedPCHs;
+
+		PrivateDependencyModuleNames.AddRange(new string[]
+		{
+			"Core",
+			"CoreUObject",
+			"CQTest",
+			"Engine",
+			"InputCore",
+			"Slate",
+			"SlateCore",
+			"SlateInspector",
+		});
+	}
+}
